Soundscape Engine Plugin Check-In
- Making Soundscape available to all projects

diff --git a/Engine/Plugins/Runtime/Soundscape/Config/FilterPlugin.ini b/Engine/Plugins/Runtime/Soundscape/Config/FilterPlugin.ini
new file mode 100644
--- /dev/null
+++ b/Engine/Plugins/Runtime/Soundscape/Config/FilterPlugin.ini
@@ -0,0 +1,8 @@
+[FilterPlugin]
+; This section lists additional files which will be packaged along with your plugin. Paths should be listed relative to the root plugin directory, and
+; may include "...", "*", and "?" wildcards to match directories, files, and individual characters respectively.
+;
+; Examples:
+;    /README.txt
+;    /Extras/...
+;    /Binaries/ThirdParty/*.dll
diff --git a/Engine/Plugins/Runtime/Soundscape/Soundscape.uplugin b/Engine/Plugins/Runtime/Soundscape/Soundscape.uplugin
new file mode 100644
--- /dev/null
+++ b/Engine/Plugins/Runtime/Soundscape/Soundscape.uplugin
@@ -0,0 +1,30 @@
+{
+	"FileVersion": 3,
+	"Version": 1,
+	"VersionName": "1.0",
+	"FriendlyName": "Soundscape Plugin",
+	"Description": "A Dynamic Ambient Sound System",
+	"Category": "Audio",
+	"CreatedBy": "Epic Games, Inc.",
+	"CreatedByURL": "http://epicgames.com",
+	"DocsURL": "",
+	"MarketplaceURL": "",
+	"SupportURL": "",
+	"EnabledByDefault": false,
+	"CanContainContent": true,
+	"IsBetaVersion": true,
+	"IsExperimentalVersion": true,
+	"Installed": false,
+	"Modules": [
+		{
+			"Name": "Soundscape",
+			"Type": "Runtime",
+			"LoadingPhase": "PreDefault"
+		},
+		{
+			"Name": "SoundscapeEditor",
+			"Type": "Editor",
+			"LoadingPhase": "Default"
+		}
+	]
+}
\ No newline at end of file
diff --git a/Engine/Plugins/Runtime/Soundscape/Source/SoundScape/Private/SoundScape.cpp b/Engine/Plugins/Runtime/Soundscape/Source/SoundScape/Private/SoundScape.cpp
new file mode 100644
--- /dev/null
+++ b/Engine/Plugins/Runtime/Soundscape/Source/SoundScape/Private/SoundScape.cpp
@@ -0,0 +1,42 @@
+// Copyright Epic Games, Inc. All Rights Reserved.
+
+#include "SoundScape.h"
+
+bool USoundscapeBPFunctionLibrary::SpawnSoundscapeColor(UObject* WorldContextObject, class USoundscapeColor* SoundscapeColorIn, UActiveSoundscapeColor*& ActiveSoundscapeColor)
+{
+	// Verify that the input SoundscapeElement is not null
+	if (SoundscapeColorIn && WorldContextObject)
+	{
+		// Create a new ActiveSoundscapeColor
+		ActiveSoundscapeColor = NewObject<UActiveSoundscapeColor>(WorldContextObject);
+
+		// Initialize parameters
+		ActiveSoundscapeColor->SetParameterValues(SoundscapeColorIn);
+
+#if WITH_EDITOR
+		// Bind to delegate for live update from the SoundscapeColor Editor
+		ActiveSoundscapeColor->BindToParameterChangeDelegate(SoundscapeColorIn);
+#endif
+
+		return true;
+	}
+
+	return false;
+}
+
+bool USoundscapeBPFunctionLibrary::SpawnSoundscapePalette(UObject* WorldContextObject, class USoundscapePalette* SoundscapePaletteIn, UActiveSoundscapePalette*& ActiveSoundscapePaletteOut)
+{
+	// Verify that the input SoundscapeElement is not null
+	if (SoundscapePaletteIn && WorldContextObject)
+	{
+		// Create a new ActiveSoundscapeColor
+		ActiveSoundscapePaletteOut = NewObject<UActiveSoundscapePalette>(WorldContextObject);
+
+		// Initialize parameters
+		ActiveSoundscapePaletteOut->InitializeSettings(WorldContextObject, SoundscapePaletteIn);
+
+		return true;
+	}
+
+	return false;
+}
diff --git a/Engine/Plugins/Runtime/Soundscape/Source/SoundScape/Private/SoundScapeModule.cpp b/Engine/Plugins/Runtime/Soundscape/Source/SoundScape/Private/SoundScapeModule.cpp
new file mode 100644
--- /dev/null
+++ b/Engine/Plugins/Runtime/Soundscape/Source/SoundScape/Private/SoundScapeModule.cpp
@@ -0,0 +1,21 @@
+// Copyright Epic Games, Inc. All Rights Reserved.
+
+#include "SoundScapeModule.h"
+#include "CoreMinimal.h"
+#include "Stats/Stats.h"
+#include "Modules/ModuleInterface.h"
+#include "Modules/ModuleManager.h"
+
+DEFINE_LOG_CATEGORY(LogSoundscape);
+
+IMPLEMENT_MODULE(FSoundscapeModule, Soundscape)
+
+void FSoundscapeModule::StartupModule()
+{
+}
+
+void FSoundscapeModule::ShutdownModule()
+{
+}
+
+
diff --git a/Engine/Plugins/Runtime/Soundscape/Source/SoundScape/Private/SoundScapePalette.cpp b/Engine/Plugins/Runtime/Soundscape/Source/SoundScape/Private/SoundScapePalette.cpp
new file mode 100644
--- /dev/null
+++ b/Engine/Plugins/Runtime/Soundscape/Source/SoundScape/Private/SoundScapePalette.cpp
@@ -0,0 +1,83 @@
+// Copyright Epic Games, Inc. All Rights Reserved.
+
+#include "SoundScapePalette.h"
+#include "SoundscapeColor.h"
+#include "Engine/World.h"
+#include "AudioDevice.h"
+#include "DrawDebugHelpers.h"
+#include "Kismet/GameplayStatics.h"
+
+USoundscapePalette::USoundscapePalette()
+{
+}
+
+void USoundscapePalette::PostLoad()
+{
+	Super::PostLoad();
+}
+
+void USoundscapePalette::Serialize(FArchive& Ar)
+{
+	Super::Serialize(Ar);
+}
+
+void UActiveSoundscapePalette::InitializeSettings(UObject* WorldContextObject, USoundscapePalette* SoundscapePalette)
+{
+	if (SoundscapePalette && WorldContextObject)
+	{
+		// Set up Soundscape Colors
+		for (FSoundscapePaletteColor& SoundscapePaletteColor : SoundscapePalette->Colors)
+		{
+			if (USoundscapeColor* SoundscapeColor = SoundscapePaletteColor.SoundscapeColor)
+			{
+				// Verify that the input SoundscapeElement is not null
+				if (SoundscapeColor)
+				{
+					// Create a new ActiveSoundscapeColor
+					UActiveSoundscapeColor* ActiveSoundscapeColor = NewObject<UActiveSoundscapeColor>(WorldContextObject);
+
+					//
+					ActiveSoundscapeColor->SetParameterValues(SoundscapeColor);
+
+#if WITH_EDITOR
+					ActiveSoundscapeColor->BindToParameterChangeDelegate(SoundscapeColor);
+#endif //WITH_EDITOR
+					ActiveSoundscapeColor->VolumeMod = SoundscapePaletteColor.ColorVolume;
+					ActiveSoundscapeColor->PitchMod = SoundscapePaletteColor.ColorPitch;
+					ActiveSoundscapeColor->FadeInMin = SoundscapePaletteColor.ColorFadeIn;
+
+					// Playback stop values
+					ActiveSoundscapeColor->FadeOutMin = SoundscapePaletteColor.ColorFadeOut;
+
+					// Add to list
+					ActiveSoundscapeColors.Add(ActiveSoundscapeColor);
+				}
+
+			}
+		}
+	}
+}
+
+void UActiveSoundscapePalette::Play()
+{
+	for (UActiveSoundscapeColor* ActiveSoundscapeColor : ActiveSoundscapeColors)
+	{
+		if (ActiveSoundscapeColor)
+		{
+			ActiveSoundscapeColor->PlayNative();
+		}
+	}
+}
+
+void UActiveSoundscapePalette::Stop()
+{
+	for (UActiveSoundscapeColor* ActiveSoundscapeColor : ActiveSoundscapeColors)
+	{
+		if (ActiveSoundscapeColor)
+		{
+			ActiveSoundscapeColor->StopNative();
+		}
+	}
+
+}
+
diff --git a/Engine/Plugins/Runtime/Soundscape/Source/SoundScape/Private/SoundscapeColor.cpp b/Engine/Plugins/Runtime/Soundscape/Source/SoundScape/Private/SoundscapeColor.cpp
new file mode 100644
--- /dev/null
+++ b/Engine/Plugins/Runtime/Soundscape/Source/SoundScape/Private/SoundscapeColor.cpp
@@ -0,0 +1,578 @@
+// Copyright Epic Games, Inc. All Rights Reserved.
+
+
+#include "SoundscapeColor.h"
+#include "UObject/WeakObjectPtr.h"
+#include "Components/AudioComponent.h"
+#include "AudioDevice.h"
+#include "TimerManager.h"
+#include "DrawDebugHelpers.h"
+#include "SoundscapeSubsystem.h"
+
+USoundscapeColor::USoundscapeColor()
+	: VolumeBase(1.0f)
+	, PitchBase(1.0f)
+{
+}
+
+void USoundscapeColor::PostLoad()
+{
+	Super::PostLoad();
+}
+
+void USoundscapeColor::Serialize(FArchive& Ar)
+{
+	Super::Serialize(Ar);
+}
+
+#if WITH_EDITORONLY_DATA
+void USoundscapeColor::PostEditChangeProperty(FPropertyChangedEvent& PropertyChangedEvent)
+{
+	// Broadcast changes to the properties so instances can update
+	OnSoundscapeColorParameterChanges.Broadcast(this);
+
+	Super::PostEditChangeProperty(PropertyChangedEvent);
+}
+#endif
+
+void UActiveSoundscapeColorVoice::BeginDestroy()
+{
+	if (UWorld* World = GetWorld())
+	{
+		// Clear out all timers
+		World->GetTimerManager().ClearTimer(TimerHandle);
+	}
+
+	Super::BeginDestroy();
+}
+
+void UActiveSoundscapeColorVoice::StopLimitedDurationVoice(float FadeOutTime)
+{
+	if (AudioComponent)
+	{
+		AudioComponent->FadeOut(FadeOutTime, 0.0f);
+	}
+}
+
+void UActiveSoundscapeColor::BeginDestroy()
+{
+	if (UWorld* World = GetWorld())
+	{
+		// Clear out all timers
+		World->GetTimerManager().ClearTimer(UpdateTimer);
+	}
+
+	Super::BeginDestroy();
+}
+
+void UActiveSoundscapeColor::SetParameterValues(const USoundscapeColor* SoundscapeColor)
+{
+	if (SoundscapeColor)
+	{
+		Sound = SoundscapeColor->Sound;
+
+		VolumeBase = SoundscapeColor->VolumeBase;
+		PitchBase = SoundscapeColor->PitchBase;
+
+		ModulationBehavior = SoundscapeColor->ModulationBehavior;
+		PlaybackBehavior = SoundscapeColor->PlaybackBehavior;
+		SpawnBehavior = SoundscapeColor->SpawnBehavior;
+	}
+}
+
+#if WITH_EDITOR
+void UActiveSoundscapeColor::BindToParameterChangeDelegate(USoundscapeColor* SoundscapeColor)
+{
+
+#if WITH_EDITORONLY_DATA
+	SoundscapeParameterChangeDelegate.BindUFunction(this, TEXT("OnSoundscapeColorParameterChange"));
+	SoundscapeColor->OnSoundscapeColorParameterChanges.Add(SoundscapeParameterChangeDelegate);
+#endif
+}
+#endif
+
+#if WITH_EDITOR
+void UActiveSoundscapeColor::OnSoundscapeColorParameterChange(const USoundscapeColor* SoundscapeColor)
+{
+	SetParameterValues(SoundscapeColor);
+}
+#endif
+
+void UActiveSoundscapeColor::PlayNative()
+{
+	if (bIsPlaying)
+	{
+		// Early out if we're already playing
+		return;
+	}
+
+	// Call internal play
+	StartPlaying();
+
+}
+
+void UActiveSoundscapeColor::StopNative()
+{
+	if (bIsPlaying == false)
+	{
+		// Early out if we've already stopped
+		return;
+	}
+
+	// Call internal stop
+	StopPlaying();
+
+}
+
+void UActiveSoundscapeColor::Play(float ColorVolume, float ColorPitch, float ColorFadeInTime)
+{
+	if (bIsPlaying)
+	{
+		// Early out if we're already playing
+		return;
+	}
+
+	// Cache play values
+	VolumeMod = ColorVolume;
+	PitchMod = ColorPitch;
+	FadeInMin = ColorFadeInTime;
+
+	// Call internal play
+	StartPlaying();
+}
+
+void UActiveSoundscapeColor::Stop(float ColorFadeOutTime)
+{
+	if (bIsPlaying == false)
+	{
+		// Early out if we've already stopped
+		return;
+	}
+
+	// Cache stop values
+	FadeOutMin = ColorFadeOutTime;
+
+	// Call internal stop
+	StopPlaying();
+}
+
+bool UActiveSoundscapeColor::IsPlaying()
+{
+	return bIsPlaying;
+}
+
+void UActiveSoundscapeColor::UpdateSoundscapeColor()
+{
+	if (bIsPlaying)
+	{
+		Update();
+	}
+}
+
+FTraceHandle UActiveSoundscapeColor::SpawnSoundByTrace()
+{
+	// Bind Trace Delegate to this object to call OnTraceCompleted
+	TraceDelegate.BindUObject(this, &UActiveSoundscapeColor::OnTraceCompleted);
+
+	if (UWorld* World = GetWorld())
+	{
+		// Get Trace start and end vectors
+		FVector TraceStart = NewSoundParams.Location;
+		FVector TraceEnd = TraceStart + (SpawnBehavior.TraceDistance * SpawnBehavior.TraceVector);
+
+		// Flag that we are now tracing (won't trigger another trace until we're done)
+		bTracing = true;
+
+		// Kick off async trace, return the handle
+		return World->AsyncLineTraceByChannel(
+			EAsyncTraceType::Single
+			, TraceStart
+			, TraceEnd
+			, SpawnBehavior.TraceChannel
+			, FCollisionQueryParams::DefaultQueryParam
+			, FCollisionResponseParams::DefaultResponseParam
+			, &TraceDelegate
+		);
+	}
+
+	return FTraceHandle();
+}
+
+void UActiveSoundscapeColor::OnTraceCompleted(const FTraceHandle& Handle, FTraceDatum& Data)
+{
+	// Set tracing flag to false now that we've completed our async trace
+	bTracing = false;
+
+	// Check if there are Data Hit Results
+	if (Data.OutHits.Num() > 0)
+	{
+		// Get first result
+		if (Data.OutHits[0].bBlockingHit)
+		{
+			// Update sound location
+			NewSoundParams.Location = Data.OutHits[0].ImpactPoint;
+
+			// Spawn new sound at updated location
+			SpawnSoundscapeColorVoice();
+		}
+	}
+}
+
+
+void UActiveSoundscapeColor::StartPlaying()
+{
+	// Update state
+	bIsPlaying = true;
+	bFirstSpawn = true;
+
+	// Set up for first Timer delay
+	float FirstDelayTime = 0.0001f;
+
+	if (SpawnBehavior.bDelayFirstSpawn)
+	{
+		// If we delay the first spawn, then 
+		FirstDelayTime = FMath::FRandRange(SpawnBehavior.MinFirstSpawnDelay, SpawnBehavior.MaxFirstSpawnDelay);
+	}
+	
+	if (UWorld* World = GetWorld())
+	{
+			
+		World->GetTimerManager().SetTimer(UpdateTimer, this, &UActiveSoundscapeColor::UpdateSoundscapeColor, FirstDelayTime);
+	}
+}
+
+void UActiveSoundscapeColor::StopPlaying()
+{
+	// Stop with appropriate fade time
+	for (UActiveSoundscapeColorVoice* SoundscapeColorVoice : SoundscapeColorVoicePool)
+	{
+		if (SoundscapeColorVoice)
+		{
+
+			if (UAudioComponent* VoiceAudioComponent = SoundscapeColorVoice->AudioComponent)
+			{
+				if (EAudioComponentPlayState::Stopped != VoiceAudioComponent->GetPlayState())
+				{
+					VoiceAudioComponent->FadeOut(FadeOutMin, 0.0f);
+				}
+
+				if (UWorld* World = GetWorld())
+				{
+					// Clear out all timers
+					World->GetTimerManager().ClearTimer(SoundscapeColorVoice->TimerHandle);
+				}
+				
+			}
+		}
+	}
+
+	bIsPlaying = false;
+}
+
+void UActiveSoundscapeColor::Update()
+{
+	UWorld* World = GetWorld();
+	bool bNeedToSpawnSound = false;
+	bVoiceInPoolFree = false;
+	int32 ActiveVoicesPlaying = 0;
+
+	FVector ListenerLocation = FVector::ZeroVector;
+	FVector ListenerForward = FVector::ZeroVector;
+	FVector ListenerUp = FVector::ZeroVector;
+	
+	if (FAudioDevice* AudioDevice = World->GetAudioDeviceRaw())
+	{
+		// Get available Listener Proxies
+		TArray<FListenerProxy>& ListenerProxies = AudioDevice->ListenerProxies;
+
+		// TODO: Handle Split Screen
+		if (ListenerProxies.Num())
+		{
+			// Get Location and Direction from Listener
+			const FTransform ListenerTransform = ListenerProxies[0].Transform;
+			ListenerLocation = ListenerTransform.GetLocation();
+			ListenerForward = ListenerTransform.GetRotation().GetForwardVector();
+			ListenerUp = ListenerTransform.GetRotation().GetUpVector();
+		}
+	}
+
+	// Free up old voices
+	for (UActiveSoundscapeColorVoice* SoundscapeColorVoice : SoundscapeColorVoicePool)
+	{
+		if (SoundscapeColorVoice)
+		{
+			
+			if (UAudioComponent* VoiceAudioComponent = SoundscapeColorVoice->AudioComponent)
+			{
+				// Get current distance from listener
+				FVector DistanceRelativeToListener = VoiceAudioComponent->GetComponentTransform().GetLocation() - ListenerLocation;
+				DistanceRelativeToListener.Z = (SpawnBehavior.bCullByMaxSpawnDistance && SpawnBehavior.bIgnoreZAxisWhenCullingByMaxDistance) ? 0.0f : DistanceRelativeToListener.Z;
+				const float DistanceFromListener = DistanceRelativeToListener.Length();
+
+				// If Color Voice is out of range by more than 1% of the max distance, stop the sound
+				if(DistanceFromListener > (SpawnBehavior.MaxSpawnDistance * 1.01f) && SpawnBehavior.bCullByMaxSpawnDistance)
+				{
+					const float FadeTime = FMath::Max(1.0f, ModulationBehavior.MinFadeOutTime);
+					VoiceAudioComponent->FadeOut(FadeTime, 0.0f);
+				}
+
+				if (EAudioComponentPlayState::Stopped == VoiceAudioComponent->GetPlayState())
+				{
+					// Voice is not in use, mark it as free to use
+					SoundscapeColorVoice->bFree = true;
+
+					// Flag that a voice is available
+					bVoiceInPoolFree = true;
+				}
+				else
+				{
+					// Voice is in use, increment active voice count
+					++ActiveVoicesPlaying;
+				}
+			}
+		}
+	}
+
+	// Determine if conditions are right to play a new sound
+	if (ActiveVoicesPlaying < SpawnBehavior.MaxNumberOfSpawnedElements)
+	{
+		bNeedToSpawnSound = true;
+	}
+
+	// Set up timer if needed
+	if (SpawnBehavior.bContinuouslyRespawn)
+	{
+		// Set random range spawn time, with a minimum value of 0.0001f (which is basically next frame)
+		float SpawnDelayTime = FMath::Max(0.0001f, FMath::FRandRange(SpawnBehavior.MinSpawnDelay, SpawnBehavior.MaxSpawnDelay));
+
+		// Set off timer again
+		if (World)
+		{
+			World->GetTimerManager().SetTimer(UpdateTimer, this, &UActiveSoundscapeColor::UpdateSoundscapeColor, SpawnDelayTime);
+		}
+	}
+
+
+
+	// Play Sound if needed and valid sound available
+	if (Sound && bNeedToSpawnSound && World)
+	{
+		// Calculate Spawn Parameters
+		CalculateSpawnParams();
+
+		USoundscapeSubsystem* SoundscapeSubsystem = nullptr;
+
+		// Set flag for can spawn
+		bool bCanSpawnSound = true;
+
+		// Get Soundscape Subsystem
+		if (UGameInstance* GameInstance = World->GetGameInstance())
+		{
+			SoundscapeSubsystem = GameInstance->GetSubsystem<USoundscapeSubsystem>();
+		}
+
+		// Check to see if Filter By Color Point Density is TRUE, determine if Soundscape Subsystem is still valid
+		if (SpawnBehavior.bFilterByColorPointDensity && SoundscapeSubsystem)
+		{
+			// Cache Color Point Density of specified type
+			int32 ColorPointDensity = SoundscapeSubsystem->CheckColorPointDensity(NewSoundParams.Location, SpawnBehavior.ColorPoint);
+
+			// If Color Point Density is under threshold, set spawn flag to false
+			if (SpawnBehavior.MinColorPointNumber > ColorPointDensity)
+			{
+				bCanSpawnSound = false;
+			}
+
+			// Attempt Draw Debug Cell
+			SoundscapeSubsystem->DrawDebugCell(NewSoundParams.Location, bCanSpawnSound);
+		}
+
+		// If Sound can Spawn, Spawn
+		if(bCanSpawnSound)
+		{
+			// If not filtering by Color Point Density, then check to see if Position By Trace
+			if (SpawnBehavior.bPositionByTrace == false)
+			{
+				//Spawn Color Voice
+				SpawnSoundscapeColorVoice();
+			}
+			else if (bTracing == false)
+			{
+				//If we're not currently mid-trace, spawn sound by trace
+				TraceHandle = SpawnSoundByTrace();
+			}
+		}
+
+	}
+
+
+	// No longer the first time
+	bFirstSpawn = false;
+}
+
+void UActiveSoundscapeColor::CalculateSpawnParams()
+{
+	if (UWorld* World = GetWorld())
+	{
+		// Play Sound if needed and valid sound available
+		if (Sound)
+		{
+			// Get the new sound's volume
+			float NewRandVolume = ModulationBehavior.bRandomizeVolume ? FMath::FRandRange(ModulationBehavior.VolumeMin, ModulationBehavior.VolumeMax) : 1.0f;
+			NewSoundParams.Volume = VolumeMod * VolumeBase * NewRandVolume;
+
+			// Get the new sound's pitch
+			float NewRandPitch = ModulationBehavior.bRandomizePitch ? FMath::FRandRange(ModulationBehavior.PitchMin, ModulationBehavior.PitchMax) : 1.0f;
+			NewSoundParams.Pitch = PitchMod * PitchBase * NewRandPitch;
+
+			// Get the new sound's fade in time
+			float NewRandFadeIn = ModulationBehavior.bFadeVolume ? FMath::FRandRange(ModulationBehavior.MinFadeInTime, ModulationBehavior.MaxFadeInTime) : 0.0f;
+			NewSoundParams.FadeInTime = bFirstSpawn ? FMath::Max(FadeInMin, NewRandFadeIn) : NewRandFadeIn;
+			NewSoundParams.FadeInTime = (bFirstSpawn && ModulationBehavior.bOnlyFadeInOnRetrigger) ? 0.0f : NewSoundParams.FadeInTime;
+
+			// Get the new sound's start time
+			NewSoundParams.StartTime = PlaybackBehavior.bRandomizeStartingSeekTime ? FMath::FRandRange(0.0f, Sound->Duration) : 0.0f;
+
+			FVector ListenerLocation;
+			FVector ListenerForward;
+			FVector ListenerUp;
+
+			if (FAudioDevice* AudioDevice = World->GetAudioDeviceRaw())
+			{
+				// Get available Listener Proxies
+				TArray<FListenerProxy>& ListenerProxies = AudioDevice->ListenerProxies;
+
+				// TODO: Handle Split Screen
+				if (ListenerProxies.Num())
+				{
+					// Get Location and Direction from Listener
+					const FTransform ListenerTransform = ListenerProxies[0].Transform;
+					ListenerLocation = ListenerTransform.GetLocation();
+					ListenerForward = ListenerTransform.GetRotation().GetForwardVector();
+					ListenerUp = ListenerTransform.GetRotation().GetUpVector();
+				}
+			}
+
+			// Get random distance based on range
+			float NewSpawnDistance = FMath::RandRange(SpawnBehavior.MinSpawnDistance, SpawnBehavior.MaxSpawnDistance);
+
+			// Get angle from min/max range then
+			float NewSpawnAngle = FMath::RandRange(SpawnBehavior.MinSpawnAngle, SpawnBehavior.MaxSpawnAngle);
+
+			// Set up altitudinal clamp offset
+			float AltitudinalClampOffset = 0.0f;
+
+			if (SpawnBehavior.ClampMode == ESoundscapeColorAltitudeClampMode::Relative)
+			{
+				AltitudinalClampOffset = ListenerLocation.Z;
+			}
+
+			// Sometimes Left Sometimes Right
+			if (FMath::RandBool())
+			{
+				NewSpawnAngle = NewSpawnAngle * (-1.0f);
+			}
+
+			// Get X Y Vector
+			NewSoundParams.Location = ListenerForward.RotateAngleAxis(NewSpawnAngle, ListenerUp);
+
+			// Random Z Vector
+			NewSoundParams.Location.Z = FMath::RandRange(-1.0f, 1.0f);
+
+			// Scale Vector and add to Listener Location
+			NewSoundParams.Location = (NewSoundParams.Location * NewSpawnDistance) + ListenerLocation;
+
+			// Clamp the Z Altitude
+			NewSoundParams.Location.Z = SpawnBehavior.bClampHeight ? FMath::Clamp(NewSoundParams.Location.Z, AltitudinalClampOffset + SpawnBehavior.MinSpawnHeightClamp, AltitudinalClampOffset + SpawnBehavior.MaxSpawnHeightClamp) : NewSoundParams.Location.Z;
+
+			if (SpawnBehavior.bRotateSoundSource)
+			{
+				NewSoundParams.Rotation.Yaw = FMath::Max(0.0f, FMath::FRandRange(SpawnBehavior.MinAzimuthalRotationAngle, SpawnBehavior.MaxAzimuthalRotationAngle));
+				NewSoundParams.Rotation.Pitch = FMath::Max(0.0f, FMath::FRandRange(SpawnBehavior.MinAltitudinalRotationAngle, SpawnBehavior.MaxAltitudinalRotationAngle));
+			}
+		}
+	}
+}
+
+// Spawn a Color Voice
+void UActiveSoundscapeColor::SpawnSoundscapeColorVoice()
+{
+	if (UWorld* World = GetWorld())
+	{
+		// Play Sound if needed and valid sound available
+		if (Sound)
+		{
+			// Set up Voice and Component
+			UActiveSoundscapeColorVoice* NewActiveSoundscapeColorVoice = nullptr;
+			UAudioComponent* NewAudioComponent = nullptr;
+
+			// If there's a free voice in the pool, get first available voice
+			if (bVoiceInPoolFree)
+			{
+				for (UActiveSoundscapeColorVoice* ActiveSoundscapeColorVoice : SoundscapeColorVoicePool)
+				{
+					if (ActiveSoundscapeColorVoice)
+					{
+						// Get first free voice
+						if (ActiveSoundscapeColorVoice->bFree)
+						{
+							// Cache pointer to this voice
+							NewActiveSoundscapeColorVoice = ActiveSoundscapeColorVoice;
+
+							break;
+						}
+
+					}
+				}
+			}
+
+			// If the ActiveSoundscapeColorVoice is still null, create a new one
+			if (NewActiveSoundscapeColorVoice == nullptr)
+			{
+				// Create new voice
+				NewActiveSoundscapeColorVoice = NewObject<UActiveSoundscapeColorVoice>(World);
+
+				// Add new voice to the pool
+				SoundscapeColorVoicePool.Add(NewActiveSoundscapeColorVoice);
+			}
+
+			// If the Audio Component invalid, create a new one
+			if (NewActiveSoundscapeColorVoice->AudioComponent == nullptr)
+			{
+				NewActiveSoundscapeColorVoice->AudioComponent = NewObject<UAudioComponent>(World);
+			}
+
+			// Cache Audio Component ptr
+			NewAudioComponent = NewActiveSoundscapeColorVoice->AudioComponent;
+
+			// Set World location of Audio Component
+			NewAudioComponent->SetWorldLocationAndRotation(NewSoundParams.Location, NewSoundParams.Rotation.Quaternion());
+
+			// Set relevant Audio Component values before playing
+			NewAudioComponent->SetPitchMultiplier(NewSoundParams.Pitch);
+			NewAudioComponent->SetSound(Sound);
+
+			NewAudioComponent->bAutoDestroy = false;
+
+			// Play the sound
+			NewAudioComponent->FadeIn(NewSoundParams.FadeInTime, NewSoundParams.Volume, NewSoundParams.StartTime);
+
+			// Voice is no longer free
+			NewActiveSoundscapeColorVoice->bFree = false;
+
+			// Handle limited duration sounds
+			if (PlaybackBehavior.bLimitPlaybackDuration)
+			{
+				float FadeOutDuration = FMath::Max(0.0001f, FMath::FRandRange(ModulationBehavior.MinFadeOutTime, ModulationBehavior.MaxFadeOutTime));
+				float LimitedDuration = FMath::Max(0.0001f, FMath::FRandRange(PlaybackBehavior.MinPlaybackDuration, PlaybackBehavior.MaxPlaybackDuration));
+				FTimerDelegate TimerDelegate;
+
+				TimerDelegate.BindUFunction(NewActiveSoundscapeColorVoice, FName("StopLimitedDurationVoice"), FadeOutDuration);
+				GetWorld()->GetTimerManager().SetTimer(NewActiveSoundscapeColorVoice->TimerHandle, TimerDelegate, LimitedDuration, false);
+			}
+
+		}
+
+
+	}
+}
\ No newline at end of file
diff --git a/Engine/Plugins/Runtime/Soundscape/Source/SoundScape/Private/SoundscapeColorPoint.cpp b/Engine/Plugins/Runtime/Soundscape/Source/SoundScape/Private/SoundscapeColorPoint.cpp
new file mode 100644
--- /dev/null
+++ b/Engine/Plugins/Runtime/Soundscape/Source/SoundScape/Private/SoundscapeColorPoint.cpp
@@ -0,0 +1,65 @@
+// Copyright Epic Games, Inc. All Rights Reserved.
+
+
+#include "SoundscapeColorPoint.h"
+#include "SoundscapeSubsystem.h"
+#include "GameFramework/Actor.h"
+#include "Engine/World.h"
+#include "Engine/GameInstance.h"
+
+// Sets default values for this component's properties
+USoundscapeColorPointComponent::USoundscapeColorPointComponent()
+{
+	// Set this component to be initialized when the game starts, and to be ticked every frame.  You can turn these features
+	// off to improve performance if you don't need them.
+	PrimaryComponentTick.bCanEverTick = false;
+	// ...
+}
+
+
+void USoundscapeColorPointComponent::GetInfo(FGameplayTag& ColorPointOut, FVector& LocationOut) const
+{
+	// Get and validate Owner Actor
+	if (const AActor* Owner = GetOwner())
+	{
+		// Set Location and Color Point values
+		LocationOut = Owner->GetActorLocation();
+		ColorPointOut = ColorPoint;
+	}
+}
+
+// Called when the game starts
+void USoundscapeColorPointComponent::BeginPlay()
+{
+	// Add this Active Color Point to the Subsystem
+	if (const UWorld* World = GetWorld())
+	{
+		if (const UGameInstance* GameInstance = World->GetGameInstance())
+		{
+			if (USoundscapeSubsystem* SoundscapeSubsystem = GameInstance->GetSubsystem<USoundscapeSubsystem>())
+			{
+				SoundscapeSubsystem->AddActiveColorPoint(this);
+			}
+		}
+	}
+
+	Super::BeginPlay();
+}
+
+void USoundscapeColorPointComponent::EndPlay(const EEndPlayReason::Type EndPlayReason)
+{
+	// Remove this Active Color Point from the Subsystem
+	if (const UWorld* World = GetWorld())
+	{
+		if (const UGameInstance* GameInstance = World->GetGameInstance())
+		{
+			if (USoundscapeSubsystem* SoundscapeSubsystem = GameInstance->GetSubsystem<USoundscapeSubsystem>())
+			{
+				SoundscapeSubsystem->RemoveActiveColorPoint(this);
+			}
+		}
+	}
+
+	Super::EndPlay(EndPlayReason);
+}
+
diff --git a/Engine/Plugins/Runtime/Soundscape/Source/SoundScape/Private/SoundscapeSettings.cpp b/Engine/Plugins/Runtime/Soundscape/Source/SoundScape/Private/SoundscapeSettings.cpp
new file mode 100644
--- /dev/null
+++ b/Engine/Plugins/Runtime/Soundscape/Source/SoundScape/Private/SoundscapeSettings.cpp
@@ -0,0 +1,17 @@
+// Copyright Epic Games, Inc. All Rights Reserved.
+
+
+#include "SoundscapeSettings.h"
+
+USoundscapeSettings::USoundscapeSettings(const FObjectInitializer& ObjectInitializer)
+	: Super(ObjectInitializer)
+{
+
+}
+
+#if WITH_EDITOR
+void USoundscapeSettings::PostEditChangeProperty(FPropertyChangedEvent& PropertyChangedEvent)
+{
+	Super::PostEditChangeProperty(PropertyChangedEvent);
+}
+#endif
\ No newline at end of file
diff --git a/Engine/Plugins/Runtime/Soundscape/Source/SoundScape/Private/SoundscapeSubsystem.cpp b/Engine/Plugins/Runtime/Soundscape/Source/SoundScape/Private/SoundscapeSubsystem.cpp
new file mode 100644
--- /dev/null
+++ b/Engine/Plugins/Runtime/Soundscape/Source/SoundScape/Private/SoundscapeSubsystem.cpp
@@ -0,0 +1,1025 @@
+// Copyright Epic Games, Inc. All Rights Reserved.
+
+
+#include "SoundscapeSubsystem.h"
+#include "SoundscapeSettings.h"
+#include "SoundScapePalette.h"
+#include "AudioDevice.h"
+#include "HAL/IConsoleManager.h"
+#include "TimerManager.h"
+
+static FAutoConsoleCommandWithWorld GResetSoundscape(
+	TEXT("soundscape.ResetSoundscape"),
+	TEXT("Forces restart on all active Soundscape Palettes."),
+	FConsoleCommandWithWorldDelegate::CreateLambda([](UWorld* World)
+{
+	if (!World)
+	{
+		return;
+	}
+
+	if (const UGameInstance* GameInstance = World->GetGameInstance())
+	{
+		if (USoundscapeSubsystem* SoundscapeSubsystem = GameInstance->GetSubsystem<USoundscapeSubsystem>())
+		{
+			SoundscapeSubsystem->RestartSoundscape();
+		}
+	}
+}));
+
+void USoundscapeColorPointHashMap::InitializeHash(float HashCellSizeIn, FVector GridCenterIn)
+{
+	// Cache Hash Cell Size
+	HashCellSize = FMath::Max(HashCellSizeIn, 1.0f);
+
+	// MaxGridWidth / Hash Cell Size
+	GridWidth = FMath::FloorToInt(MaxGridWidth / HashCellSize);
+
+	// Grid Width Cubed
+	NumCells = int64(GridWidth) * int64(GridWidth) * int64(GridWidth);
+
+	// Cell Size Inverted
+	HashCellFactor = 1 / HashCellSize;
+
+	// Cache Grid Centerpoint
+	GridCenter = GridCenterIn;
+
+	// Calculate Grid Origin Offset
+	GridOriginOffset = GridCenter - FVector(float(HalfMaxGridWidth), float(HalfMaxGridWidth), float(HalfMaxGridWidth));
+}
+
+void USoundscapeColorPointHashMap::ClearHash()
+{
+	// Cycle through Color Point Hash Map
+	for (auto It = ColorPointHashMap.CreateIterator(); It; ++It)
+	{
+		// If Color Point has an existing hash map, empty it
+		if (USoundscapeColorPointHashCellDensity* ColorPointHashCell = It->Value)
+		{
+			ColorPointHashCell->ColorPointHashCellDensity.Empty();
+		}
+	}
+
+	// Empty Hash Map
+	ColorPointHashMap.Empty();
+}
+
+int32 USoundscapeColorPointHashMap::NumColorPointsInCell(const FVector& Location, const FGameplayTag ColorPoint)
+{
+	// Determine if the Hash Map contains the Color Point Density Table Pointer
+	if (USoundscapeColorPointHashCellDensity** ColorPointHashCellPtr = ColorPointHashMap.Find(ColorPoint))
+	{
+		// Validate the double pointer
+		if (USoundscapeColorPointHashCellDensity* ColorPointHashCell = *ColorPointHashCellPtr)
+		{
+			// Validate the Color Point Density Map
+			if (ColorPointHashCell != nullptr)
+			{
+				// Look up Color Point Density at that Location
+				int32* CellDensity = ColorPointHashCell->ColorPointHashCellDensity.Find(CalculateHashIndex(Location));
+
+				// If Cell Density pointer is valid, return value
+				if (CellDensity != nullptr)
+				{
+					return *CellDensity;
+				}
+			}
+		}
+	}
+
+	// Otherwise return 0
+	return 0;
+}
+
+bool USoundscapeColorPointHashMap::AddColorPointToHash(const FVector& Location, const FGameplayTag ColorPoint)
+{
+	// Add a single Color Point to the Color Point Hash Map
+	// If the Color Point Coordinates exceed the Max Grid Width, it is out of bounds of the hash table
+	// Incoming Location values are added to the Grid Origin Offset in the case that the Grid is not centered around Origin 0,0,0
+	if (FMath::Abs(GridCenter.X + Location.X) > HalfMaxGridWidth
+		|| FMath::Abs(GridCenter.Y + Location.Y) > HalfMaxGridWidth
+		|| FMath::Abs(GridCenter.Z + Location.Z) > HalfMaxGridWidth)
+	{
+		UE_LOG(LogSoundscapeSubsystem, Warning, TEXT("Location out of Hash Bounds"));
+
+		// Return false, unable to add Color Point to Hash
+		return false;
+	}
+
+	// Get ColorPoint Density Map
+	USoundscapeColorPointHashCellDensity* ColorPointDensityMap = GetColorPointDensityMap(ColorPoint);
+
+	// Create a new Color Point Hash for this Color Point if one needs to be created
+	if (ColorPointDensityMap == nullptr)
+	{
+		// Create new Color Point Density Map
+		ColorPointDensityMap = NewObject<USoundscapeColorPointHashCellDensity>(this);
+
+		// Add to Color Point Hash Map
+		ColorPointHashMap.Add(ColorPoint, ColorPointDensityMap);
+	}
+
+
+	// If a Color Point Map exists for this Color Point, use it
+	if (ColorPointDensityMap)
+	{
+		// Cache Hash Index
+		uint64 HashIndex = CalculateHashIndex(Location);
+
+		// Get potential Density value from Hash Index
+		int32* ColorPointDensity = ColorPointDensityMap->ColorPointHashCellDensity.Find(HashIndex);
+
+		// Validate Density pointer
+		if (ColorPointDensity != nullptr)
+		{
+			// Color Point already tracked at this hash cell, increment
+			++(*ColorPointDensity);
+		}
+		else
+		{
+			// Density Value invalid, new Hash Key/Value must be added at that index
+			ColorPointDensityMap->ColorPointHashCellDensity.Add(HashIndex, 1);
+		}
+	}
+
+	// Color Point successfully added
+	return true;
+}
+
+void USoundscapeColorPointHashMap::AddColorPointArrayToHash(const TArray<FVector>& Locations, const FGameplayTag ColorPoint)
+{
+
+	// Local cache for Hash Indices
+	TArray<uint64> HashIndices;
+
+	// Get ColorPoint Density Map
+	USoundscapeColorPointHashCellDensity* ColorPointDensityMap = GetColorPointDensityMap(ColorPoint);
+
+	// Create a new Color Point Hash for this Color Point if one needs to be created
+	if (ColorPointDensityMap == nullptr)
+	{
+		// Create new Color Point Density Map
+		ColorPointDensityMap = NewObject<USoundscapeColorPointHashCellDensity>(this);
+
+		// Add to Color Point Hash Map
+		ColorPointHashMap.Add(ColorPoint, ColorPointDensityMap);
+	}
+
+
+	// If a Color Point Map exists for this Color Point, use it
+	if(ColorPointDensityMap)
+	{
+		// Cycle through input Locations, add to temp array the indices calculated by the hashing function
+		for (auto IndividualLocation = Locations.CreateConstIterator(); IndividualLocation; ++IndividualLocation)
+		{
+			// Make sure incoming Location Vector is within Hash Map Bounds
+			if (FMath::Abs(GridCenter.X + IndividualLocation->X) <= HalfMaxGridWidth
+				&& FMath::Abs(GridCenter.Y + IndividualLocation->Y) <= HalfMaxGridWidth
+				&& FMath::Abs(GridCenter.Z + IndividualLocation->Z) <= HalfMaxGridWidth)
+			{
+				// Add Hash Index to temp Array
+				HashIndices.Add(CalculateHashIndex(*IndividualLocation));
+			}
+			else
+			{
+				// If Location out of Hash Map Bounds, throw it away and spit out a warning
+				UE_LOG(LogSoundscapeSubsystem, Warning, TEXT("Location out of Hash Bounds"));
+			}
+		}
+
+		// Cycle through temp array of hash indices
+		for (auto HashIndex = HashIndices.CreateConstIterator(); HashIndex; ++HashIndex)
+		{
+			// Cache index
+			uint64 NewKey = *HashIndex;
+
+			// Look for Color Point Density at that Index
+			int32* ColorPointDensity = ColorPointDensityMap->ColorPointHashCellDensity.Find(NewKey);
+
+			// If Color Point Density is valid, 
+			if (ColorPointDensity != nullptr)
+			{
+				// Increment the value at the pointer
+				++(*ColorPointDensity);
+			}
+			else
+			{
+				// Add a new Cell
+				ColorPointDensityMap->ColorPointHashCellDensity.Add(NewKey, 1);
+			}
+		}
+
+	}
+}
+
+uint64 USoundscapeColorPointHashMap::CalculateHashIndex(const FVector& Location)
+{
+	const FVector OffsetLocation = Location - GridOriginOffset;
+	uint64 HashX = uint64(FMath::FloorToInt(OffsetLocation.X * HashCellFactor));
+	uint64 HashY = uint64(FMath::FloorToInt(OffsetLocation.Y * HashCellFactor)) * uint64(GridWidth);
+	uint64 HashZ = uint64(FMath::FloorToInt(OffsetLocation.Z * HashCellFactor)) * uint64(GridWidth * GridWidth);
+
+	return HashX + HashY + HashZ;
+}
+
+void USoundscapeColorPointHashMap::SetGridCenterpoint(const FVector& Location)
+{
+	// Cache new Centerpoint
+	GridCenter = Location;
+
+	// Calculate Grid Origin Offset
+	GridOriginOffset = Location - FVector(float(HalfMaxGridWidth), float(HalfMaxGridWidth), float(HalfMaxGridWidth));
+}
+
+TMap<FGameplayTag, int32> USoundscapeColorPointHashMap::GetHashMapColorPointDensitySummary()
+{
+	TMap<FGameplayTag, int32> DensitySummary;
+
+	for (auto It = ColorPointHashMap.CreateConstIterator(); It; ++It)
+	{
+		FGameplayTag ColorPoint = It.Key();
+
+		if (ColorPoint.IsValid())
+		{
+			int32 ColorDensity = 0;
+
+			if (USoundscapeColorPointHashCellDensity* ColorPointHashLayer = It.Value())
+			{
+				for (auto Jt = ColorPointHashLayer->ColorPointHashCellDensity.CreateConstIterator(); Jt; ++Jt)
+				{
+					ColorDensity += Jt.Value();
+				}
+			}
+
+			DensitySummary.Add(ColorPoint, ColorDensity);
+		}
+	}
+
+	return DensitySummary;
+}
+
+USoundscapeColorPointHashCellDensity* USoundscapeColorPointHashMap::GetColorPointDensityMap(FGameplayTag ColorPoint)
+{
+	if (USoundscapeColorPointHashCellDensity** ColorPointHashCellPtr = ColorPointHashMap.Find(ColorPoint))
+	{
+		// Get Color Point Hash Cell Pointer
+		return *ColorPointHashCellPtr;
+	}
+
+	return nullptr;
+}
+
+const FVector USoundscapeColorPointHashMap::GetHashCellCenterpoint(FVector Location)
+{
+	const int32 LocationX = FMath::FloorToInt((Location.X - GridOriginOffset.X) / HashCellSize);
+	const int32 LocationY = FMath::FloorToInt((Location.Y - GridOriginOffset.Y) / HashCellSize);
+	const int32 LocationZ = FMath::FloorToInt((Location.Z - GridOriginOffset.Z) / HashCellSize);
+
+	const float HalfHashCellSize = HashCellSize * 0.5;
+
+	return FVector(((float)LocationX * HashCellSize) + HalfHashCellSize + GridOriginOffset.X
+				, ((float)LocationY * HashCellSize) + HalfHashCellSize + GridOriginOffset.Y
+				, ((float)LocationZ * HashCellSize) + HalfHashCellSize + GridOriginOffset.Z);
+}
+
+const float USoundscapeColorPointHashMap::GetHashCellWidth()
+{
+	return HashCellSize;
+}
+
+void USoundscapeColorPointHashMapCollection::InitializeCollection()
+{
+	// Get plugin settings on Subsystem initialization
+	if (const USoundscapeSettings* ProjectSettings = GetDefault<USoundscapeSettings>())
+	{
+		// Hash Cell Width for LOD1
+		LOD1ColorPointHashWidth = FMath::Max(ProjectSettings->LOD1ColorPointHashWidth, 1.0f);
+
+		// Hash Cell LOD1 Max Distance
+		LOD1ColorPointHashDistance = FMath::Max(ProjectSettings->LOD1ColorPointHashDistance, 1.0f);
+
+		// Hash Cell Width for LOD2
+		LOD2ColorPointHashWidth = FMath::Max(ProjectSettings->LOD2ColorPointHashWidth, 1.0f);
+
+		// Hash Cell LOD2 Max Distance
+		LOD2ColorPointHashDistance = FMath::Max(ProjectSettings->LOD2ColorPointHashDistance, 1.0f);
+
+		// Hash Cell Width for LOD3
+		LOD3ColorPointHashWidth = FMath::Max(ProjectSettings->LOD3ColorPointHashWidth, 1.0f);
+	}
+
+	// Initialize LOD Hashes
+	ColorPointHashMapLOD1 = NewObject<USoundscapeColorPointHashMap>(this);
+	ColorPointHashMapLOD1->InitializeHash(LOD1ColorPointHashWidth, FVector::ZeroVector);
+
+	ColorPointHashMapLOD2 = NewObject<USoundscapeColorPointHashMap>(this);
+	ColorPointHashMapLOD2->InitializeHash(LOD2ColorPointHashWidth, FVector::ZeroVector);
+
+	ColorPointHashMapLOD3 = NewObject<USoundscapeColorPointHashMap>(this);
+	ColorPointHashMapLOD3->InitializeHash(LOD3ColorPointHashWidth, FVector::ZeroVector);
+}
+
+void USoundscapeColorPointHashMapCollection::AddColorPointArrayToHashMapCollection(const TArray<FVector>& Locations, const FGameplayTag ColorPoint)
+{
+	if (ColorPointHashMapLOD1)
+	{
+		ColorPointHashMapLOD1->AddColorPointArrayToHash(Locations, ColorPoint);
+	}
+
+	if (ColorPointHashMapLOD2)
+	{
+		ColorPointHashMapLOD2->AddColorPointArrayToHash(Locations, ColorPoint);
+	}
+
+	if (ColorPointHashMapLOD3)
+	{
+		ColorPointHashMapLOD3->AddColorPointArrayToHash(Locations, ColorPoint);
+	}
+}
+
+void USoundscapeColorPointHashMapCollection::ClearColorPointHashMapCollection()
+{
+	if (ColorPointHashMapLOD1)
+	{
+		ColorPointHashMapLOD1->ClearHash();
+	}
+
+	if (ColorPointHashMapLOD2)
+	{
+		ColorPointHashMapLOD2->ClearHash();
+	}
+
+	if (ColorPointHashMapLOD3)
+	{
+		ColorPointHashMapLOD3->ClearHash();
+	}
+}
+
+int32 USoundscapeColorPointHashMapCollection::GetColorPointHashMapCollectionDensity(const FVector Location, const FGameplayTag ColorPoint, const ESoundscapeLOD SoundscapeLOD)
+{
+	if (ColorPointHashMapLOD1 && ColorPointHashMapLOD2 && ColorPointHashMapLOD3)
+	{
+		switch (SoundscapeLOD)
+		{
+		case ESoundscapeLOD::LOD1:
+			return ColorPointHashMapLOD1->NumColorPointsInCell(Location, ColorPoint);
+
+		case ESoundscapeLOD::LOD2:
+			return ColorPointHashMapLOD2->NumColorPointsInCell(Location, ColorPoint);
+
+		case ESoundscapeLOD::LOD3:
+			return ColorPointHashMapLOD3->NumColorPointsInCell(Location, ColorPoint);
+		}
+	}
+
+	return 0;
+}
+
+void USoundscapeColorPointHashMapCollection::CalculateTotalColorPointDensity(TMap<FGameplayTag, int32>& TotalHashMapColorPointDensity, ESoundscapeLOD SoundscapeLOD)
+{
+	if (ColorPointHashMapLOD1 && ColorPointHashMapLOD2 && ColorPointHashMapLOD3)
+	{
+		switch (SoundscapeLOD)
+		{
+		case ESoundscapeLOD::LOD1:
+			TotalHashMapColorPointDensity = ColorPointHashMapLOD1->GetHashMapColorPointDensitySummary();
+
+		case ESoundscapeLOD::LOD2:
+			TotalHashMapColorPointDensity = ColorPointHashMapLOD2->GetHashMapColorPointDensitySummary();
+
+		case ESoundscapeLOD::LOD3:
+			TotalHashMapColorPointDensity = ColorPointHashMapLOD3->GetHashMapColorPointDensitySummary();
+		}
+	}
+}
+
+void USoundscapeSubsystem::Initialize(FSubsystemCollectionBase& Collection)
+{
+	UE_LOG(LogSoundscapeSubsystem, Verbose, TEXT("Subsystem Initializing"));
+
+	// Get Subsystem World
+	UWorld* World = GetWorld();
+
+	// Initialize AudioDeviceID
+	AudioDeviceID = INDEX_NONE;
+
+	// Get world
+	if (World)
+	{
+		// Get Audio Device Handle
+		FAudioDevice* AudioDevice = World->GetAudioDeviceRaw();
+
+		if (AudioDevice)
+		{
+			AudioDeviceID = AudioDevice->DeviceID;
+		}
+
+	}
+
+	// Get plugin settings on Subsystem initialization
+	if (const USoundscapeSettings* ProjectSettings = GetDefault<USoundscapeSettings>())
+	{
+		if (ProjectSettings->SoundscapePaletteCollection.IsEmpty() == false)
+		{
+			FSoundscapePaletteCollection PaletteCollectionToAdd;
+
+			PaletteCollectionToAdd.SoundscapePaletteCollection = ProjectSettings->SoundscapePaletteCollection;
+
+			UE_LOG(LogSoundscapeSubsystem, Verbose, TEXT("Project settings Soundscape Palette Collection is valid"));
+
+			AddPaletteCollection(FName("GlobalPaletteCollection"), PaletteCollectionToAdd);			
+		}
+		else
+		{
+			UE_LOG(LogSoundscapeSubsystem, Verbose, TEXT("Project settings Soundscape Palette Collection is empty"));
+		}
+
+		bDebugMode = ProjectSettings->bDebugDraw;
+
+		// Hash Cell Width for LOD1
+		LOD1ColorPointHashWidth = FMath::Max(ProjectSettings->LOD1ColorPointHashWidth, 1.0f);
+
+		// Hash Cell LOD1 Max Distance
+		LOD1ColorPointHashDistance = FMath::Max(ProjectSettings->LOD1ColorPointHashDistance, 1.0f);
+
+		// Hash Cell Width for LOD2
+		LOD2ColorPointHashWidth = FMath::Max(ProjectSettings->LOD2ColorPointHashWidth, 1.0f);
+
+		// Hash Cell LOD2 Max Distance
+		LOD2ColorPointHashDistance = FMath::Max(ProjectSettings->LOD2ColorPointHashDistance, 1.0f);
+
+		// Hash Cell Width for LOD3
+		LOD3ColorPointHashWidth = FMath::Max(ProjectSettings->LOD3ColorPointHashWidth, 1.0f);
+
+		// Hash Cell Width for the Active Hash
+		ActiveColorPointHashWidth = FMath::Max(ProjectSettings->ActiveColorPointHashWidth, 1.0f);
+
+		// Hash Cell Update Timing for the Active Hash
+		ActiveColorPointHashUpdateTimeSeconds = FMath::Max(ProjectSettings->ActiveColorPointHashUpdateTimeSeconds, 0.2f);
+	}
+	else
+	{
+		UE_LOG(LogSoundscapeSubsystem, Verbose, TEXT("Invalid Project Settings"));
+	}
+
+	// Initialize LOD Hashes
+	ColorPointHashMapLOD1 = NewObject<USoundscapeColorPointHashMap>(this);
+	ColorPointHashMapLOD1->InitializeHash(LOD1ColorPointHashWidth, FVector::ZeroVector);
+
+	ColorPointHashMapLOD2 = NewObject<USoundscapeColorPointHashMap>(this);
+	ColorPointHashMapLOD2->InitializeHash(LOD2ColorPointHashWidth, FVector::ZeroVector);
+
+	ColorPointHashMapLOD3 = NewObject<USoundscapeColorPointHashMap>(this);
+	ColorPointHashMapLOD3->InitializeHash(LOD3ColorPointHashWidth, FVector::ZeroVector);
+
+	// Initialize Active Hash Map
+	ActiveColorPointHashMap = NewObject<USoundscapeColorPointHashMap>(this);
+	ActiveColorPointHashMap->InitializeHash(ActiveColorPointHashWidth, FVector::ZeroVector);
+
+	UpdateActiveColorPointHashMap();
+}
+
+void USoundscapeSubsystem::Deinitialize()
+{
+	UE_LOG(LogSoundscapeSubsystem, Verbose, TEXT("Subsystem Deinitializing"));
+
+	ActiveColorPointUpdateTimer.Invalidate();
+}
+
+bool USoundscapeSubsystem::ShouldCreateSubsystem(UObject* Outer) const
+{
+	// Check with parent first
+	if (Super::ShouldCreateSubsystem(Outer) == false)
+	{
+		UE_LOG(LogSoundscapeSubsystem, Verbose, TEXT("Subsystem will not be created based on parent rules"));
+
+		return false;
+	}
+
+	// Get World
+	if (UWorld* OuterWorld = Outer->GetWorld())
+	{
+		if (FAudioDevice* OuterAudioDevice = OuterWorld->GetAudioDeviceRaw())
+		{
+			UE_LOG(LogSoundscapeSubsystem, Verbose, TEXT("Audio device valid, will create subsystem"));
+
+			// If we do have an audio device, we can create this subsystem.
+			return true;
+		}
+	}
+
+	UE_LOG(LogSoundscapeSubsystem, Verbose, TEXT("Subsystem will not be created"));
+
+	// If we do not have an audio device, we do not need to create this subsystem.
+	return false;
+}
+
+void USoundscapeSubsystem::SetState(FGameplayTag SoundscapeState)
+{
+	if (SoundscapeState.IsValid())
+	{
+		// Add new state to the container
+		SubsystemState.AddLeafTag(SoundscapeState);
+
+		UE_LOG(LogSoundscapeSubsystem, Verbose, TEXT("Set State called"));
+
+		UpdateState();
+	}
+	else
+	{
+		UE_LOG(LogSoundscapeSubsystem, Warning, TEXT("Set State called with invalid State Tag"));
+	}
+}
+
+void USoundscapeSubsystem::ClearState(FGameplayTag SoundscapeState)
+{
+	if (SoundscapeState.IsValid())
+	{
+		// Add new state to the container
+		SubsystemState.RemoveTag(SoundscapeState);
+
+		UE_LOG(LogSoundscapeSubsystem, Verbose, TEXT("Clear State called"));
+
+		UpdateState();
+	}
+	else
+	{
+		UE_LOG(LogSoundscapeSubsystem, Verbose, TEXT("Clear State called with invalid State Tag"));
+	}
+}
+
+void USoundscapeSubsystem::RestartSoundscape()
+{
+	// Restart matching valid Soundscape Palettes
+	for (const TPair<USoundscapePalette*, UActiveSoundscapePalette*> ActiveSoundscapePalettePair : ActivePalettes)
+	{
+		USoundscapePalette* SoundscapePalette = ActiveSoundscapePalettePair.Key;
+		UActiveSoundscapePalette* ActiveSoundscapePalette = ActiveSoundscapePalettePair.Value;
+
+		if (SoundscapePalette && ActiveSoundscapePalette)
+		{
+			ActiveSoundscapePalette->Stop();
+			ActiveSoundscapePalette->Play();
+		}
+	}
+}
+
+bool USoundscapeSubsystem::AddPaletteCollection(FName PaletteCollectionName, FSoundscapePaletteCollection PaletteCollection)
+{
+	if (UnloadedPaletteCollections.Contains(PaletteCollectionName) == false)
+	{
+		UnloadedPaletteCollections.Add(PaletteCollectionName, PaletteCollection);
+
+		UE_LOG(LogSoundscapeSubsystem, Verbose, TEXT("Palette Collection %s added"), *PaletteCollectionName.ToString());
+
+		bool bLoaded = LoadPaletteCollection(PaletteCollectionName);
+
+		UpdateState();
+
+		return bLoaded;
+	}
+
+	UE_LOG(LogSoundscapeSubsystem, Verbose, TEXT("Add Palette Collection called with empty Palette Collection"));
+
+	return false;
+}
+
+bool USoundscapeSubsystem::RemovePaletteCollection(FName PaletteCollectionName)
+{
+	if (UnloadedPaletteCollections.Find(PaletteCollectionName))
+	{
+		if (UnloadPaletteCollection(PaletteCollectionName))
+		{
+			// Unloaded, now remove
+			UnloadedPaletteCollections.Remove(PaletteCollectionName);
+
+			UE_LOG(LogSoundscapeSubsystem, Verbose, TEXT("Palette Collection %s removed"), *PaletteCollectionName.ToString());
+
+			UpdateState();
+
+			return true;
+		}
+	}
+
+	UE_LOG(LogSoundscapeSubsystem, Verbose, TEXT("Remove Palette Collection called, but the name could not be found"));
+
+	return false;
+}
+
+bool USoundscapeSubsystem::LoadPaletteCollection(FName PaletteCollectionName)
+{
+	if (UnloadedPaletteCollections.Contains(PaletteCollectionName))
+	{
+		FSoundscapePaletteCollection PaletteCollectionToLoad = UnloadedPaletteCollections.FindRef(PaletteCollectionName);
+
+		UE_LOG(LogSoundscapeSubsystem, Verbose, TEXT("Attempting to load Palette Collection %s"), *PaletteCollectionName.ToString());
+
+		for (FSoftObjectPath& ObjPath : PaletteCollectionToLoad.SoundscapePaletteCollection)
+		{
+			if (UObject* PalettePath = ObjPath.TryLoad())
+			{
+				USoundscapePalette* SoundscapePalette = Cast<USoundscapePalette>(PalettePath);
+
+				if (SoundscapePalette)
+				{
+					// If palette is valid, add it to the Subsystem Collection
+					LoadedPaletteCollectionSet.Add(SoundscapePalette);
+
+					UE_LOG(LogSoundscapeSubsystem, Verbose, TEXT("Valid Soundscape Palette %s in Palette Collection %s"), *SoundscapePalette->GetFullName(), *PaletteCollectionName.ToString());
+				}
+				else
+				{
+					UE_LOG(LogSoundscapeSubsystem, Verbose, TEXT("Invalid Soundscape Palette in Palette Collection %s"), *PaletteCollectionName.ToString());
+				}
+			}
+		}
+
+		return true;
+	}
+
+	UE_LOG(LogSoundscapeSubsystem, Verbose, TEXT("Attempted to load Palette Collection %s but failed"), *PaletteCollectionName.ToString());
+
+	return false;
+}
+
+bool USoundscapeSubsystem::UnloadPaletteCollection(FName PaletteCollectionName)
+{
+	if (UnloadedPaletteCollections.Find(PaletteCollectionName))
+	{
+		UE_LOG(LogSoundscapeSubsystem, Verbose, TEXT("Attempting to unload Palette Collection %s"), *PaletteCollectionName.ToString());
+
+		// Get unloaded set to remove from list
+		FSoundscapePaletteCollection PaletteCollectionToUnload = UnloadedPaletteCollections.FindRef(PaletteCollectionName);
+		TSet<FSoftObjectPath> PaletteCollectionSetToUnload;
+		PaletteCollectionSetToUnload.Append(PaletteCollectionToUnload.SoundscapePaletteCollection);
+		TSet<FSoftObjectPath> PaletteCollectionSetToKeep;
+
+		// Go through all unloaded collections and construct a unique list of Soft Object Paths
+		for (TPair<FName, FSoundscapePaletteCollection>& UnloadedCollectionName : UnloadedPaletteCollections)
+		{
+			if (UnloadedCollectionName.Key != PaletteCollectionName)
+			{
+				PaletteCollectionSetToKeep.Append(UnloadedCollectionName.Value.SoundscapePaletteCollection.Intersect(PaletteCollectionSetToUnload));
+			}
+		}
+
+		// Find the difference between what we want to retain and what we want to unload
+		PaletteCollectionSetToUnload = PaletteCollectionSetToUnload.Difference(PaletteCollectionSetToKeep);
+
+		// Temp set of palettes to remove and stop
+		TSet<USoundscapePalette*> PalettesToUnloadStopAndRemove;
+
+		// Go through collection, get their loaded pointers, add to temp removal set
+		for (FSoftObjectPath& ObjPath : PaletteCollectionSetToUnload)
+		{
+			if (UObject* PalettePath = ObjPath.TryLoad())
+			{
+				USoundscapePalette* SoundscapePalette = Cast<USoundscapePalette>(PalettePath);
+
+				if (SoundscapePalette)
+				{
+					// If palette is valid, add it to the Subsystem Collection
+					PalettesToUnloadStopAndRemove.Add(SoundscapePalette);
+				}
+			}
+		}
+
+		// Stop and remove from Active list
+		for (USoundscapePalette* PaletteKeyToStop : PalettesToUnloadStopAndRemove)
+		{
+			if (PaletteKeyToStop)
+			{
+				if (UActiveSoundscapePalette* ActiveSoundscapePalette = ActivePalettes.FindRef(PaletteKeyToStop))
+				{
+					ActiveSoundscapePalette->Stop();
+
+					ActivePalettes.Remove(PaletteKeyToStop);
+
+					UE_LOG(LogSoundscapeSubsystem, Verbose, TEXT("Soundscape Palette %s stopped"), *ActiveSoundscapePalette->GetFullName());
+				}
+			}
+		}
+
+		// Remove from main list of loaded palettes
+		LoadedPaletteCollectionSet = LoadedPaletteCollectionSet.Difference(PalettesToUnloadStopAndRemove);
+
+		UE_LOG(LogSoundscapeSubsystem, Verbose, TEXT("Palette Collection %s unloaded"), *PaletteCollectionName.ToString());
+
+		return true;
+	}
+
+	UE_LOG(LogSoundscapeSubsystem, Verbose, TEXT("Attempted to unload Palette Collection %s but could not find matching collection"), *PaletteCollectionName.ToString());
+
+	return false;
+}
+
+void USoundscapeSubsystem::AddColorPointCollection(FName ColorPointCollectionName, FSoundscapeColorPointCollection ColorPointCollection)
+{
+	ColorPointCollections.Add(ColorPointCollectionName, ColorPointCollection);
+
+	UpdateColorPointHashMap(*ColorPointHashMapLOD1);
+	UpdateColorPointHashMap(*ColorPointHashMapLOD2);
+	UpdateColorPointHashMap(*ColorPointHashMapLOD3);
+}
+
+bool USoundscapeSubsystem::RemoveColorPointCollection(FName ColorPointCollectionName)
+{
+	return (ColorPointCollections.Remove(ColorPointCollectionName) > 0);
+}
+
+int32 USoundscapeSubsystem::CheckColorPointDensity(FVector Location, FGameplayTag ColorPoint)
+{
+	FVector ListenerLocation;
+	int32 ColorPointDensity = 0;
+	ESoundscapeLOD CurrentLOD = ESoundscapeLOD::LOD1;
+
+	if (UWorld* World = GetWorld())
+	{
+		if (FAudioDevice* AudioDevice = World->GetAudioDeviceRaw())
+		{
+			// Get available Listener Proxies
+			TArray<FListenerProxy>& ListenerProxies = AudioDevice->ListenerProxies;
+
+			// TODO: Handle Split Screen
+			if (ListenerProxies.Num())
+			{
+				// Get Location and Direction from Listener
+				FTransform& ListenerTransform = ListenerProxies[0].Transform;
+				ListenerLocation = ListenerTransform.GetLocation();
+			}
+		}
+	}
+
+	// Get relative distance from the Listener Location
+	FVector RelativeLocation = Location - ListenerLocation;
+
+	// If relative distance is less than LOD 1 cutoff distance, use LOD 1 Map
+	if (RelativeLocation.Length() < LOD1ColorPointHashDistance)
+	{
+		CurrentLOD = ESoundscapeLOD::LOD1;
+		ColorPointDensity = ColorPointDensity + ColorPointHashMapLOD1->NumColorPointsInCell(Location, ColorPoint);
+	}
+	else if (RelativeLocation.Length() < LOD2ColorPointHashDistance)
+	{
+		CurrentLOD = ESoundscapeLOD::LOD2;
+		// If relative distance is less than LOD 2 cutoff distance, use LOD 2 Map
+		ColorPointDensity = ColorPointDensity + ColorPointHashMapLOD2->NumColorPointsInCell(Location, ColorPoint);
+	}
+	else
+	{
+		CurrentLOD = ESoundscapeLOD::LOD3;
+		// If distance exceeds both cutoff distances, use LOD 3 Map
+		ColorPointDensity = ColorPointDensity + ColorPointHashMapLOD3->NumColorPointsInCell(Location, ColorPoint);
+	}
+
+	// Check Active Color Point Map
+	ColorPointDensity = ColorPointDensity + ActiveColorPointHashMap->NumColorPointsInCell(Location, ColorPoint);
+
+	// Check Color Point Hash Map Collections
+	if (ColorPointHashMapCollections.Num() > 0)
+	{
+		for (auto It = ColorPointHashMapCollections.CreateConstIterator(); It; ++It)
+		{
+			if (It)
+			{
+				USoundscapeColorPointHashMapCollection* HashMapCollection = *It;
+
+				if (HashMapCollection)
+				{
+					ColorPointDensity = ColorPointDensity + HashMapCollection->GetColorPointHashMapCollectionDensity(Location, ColorPoint, CurrentLOD);
+				}
+			}
+		}
+	}
+
+	// Return total
+	return ColorPointDensity;
+}
+
+void USoundscapeSubsystem::AddColorPointHashMapCollection(USoundscapeColorPointHashMapCollection* ColorPointHashMapCollection)
+{
+	if (ColorPointHashMapCollection)
+	{
+		ColorPointHashMapCollections.Add(ColorPointHashMapCollection);
+	}
+}
+
+bool USoundscapeSubsystem::RemoveColorPointHashMapCollection(USoundscapeColorPointHashMapCollection* ColorPointHashMapCollection)
+{
+	if (ColorPointHashMapCollection)
+	{
+		ColorPointHashMapCollections.RemoveSingleSwap(ColorPointHashMapCollection, true);
+
+		return true;
+	}
+
+	return false;
+}
+
+void USoundscapeSubsystem::DrawDebugCell(FVector Location, bool bSuccess)
+{
+	if (bDebugMode)
+	{
+		FVector ListenerLocation;
+		TPair<FVector, FVector> BoxDimensions;
+
+		if (UWorld* World = GetWorld())
+		{
+			if (FAudioDevice* AudioDevice = World->GetAudioDeviceRaw())
+			{
+				// Get available Listener Proxies
+				TArray<FListenerProxy>& ListenerProxies = AudioDevice->ListenerProxies;
+
+				// TODO: Handle Split Screen
+				if (ListenerProxies.Num())
+				{
+					// Get Location and Direction from Listener
+					FTransform& ListenerTransform = ListenerProxies[0].Transform;
+					ListenerLocation = ListenerTransform.GetLocation();
+				}
+			}
+		}
+
+		// Get relative distance from the Listener Location
+		FVector RelativeLocation = Location - ListenerLocation;
+
+		// If relative distance is less than LOD 1 cutoff distance, use LOD 1 Map
+		if (RelativeLocation.Length() < LOD1ColorPointHashDistance)
+		{
+			BoxDimensions = CalculateDebugCellDimensions(Location, ESoundscapeLOD::LOD1);
+		}
+		else if (RelativeLocation.Length() < LOD2ColorPointHashDistance)
+		{
+			BoxDimensions = CalculateDebugCellDimensions(Location, ESoundscapeLOD::LOD2);
+		}
+		else
+		{
+			BoxDimensions = CalculateDebugCellDimensions(Location, ESoundscapeLOD::LOD3);
+		}
+
+		FColor BoxColor = bSuccess ? FColor::Green : FColor::Red;
+		DrawDebugBox(GetWorld(), BoxDimensions.Key, BoxDimensions.Value, BoxColor, false, 5.0f, '\000', 5.0f);
+	}
+}
+
+TPair<FVector, FVector> USoundscapeSubsystem::CalculateDebugCellDimensions(FVector Location, ESoundscapeLOD SoundscapeLOD)
+{
+	USoundscapeColorPointHashMap* LODHashMap = nullptr;
+	TPair<FVector, FVector> BoxResults;
+
+	// If relative distance is less than LOD 1 cutoff distance, use LOD 1 Map
+	switch (SoundscapeLOD)
+	{
+	case ESoundscapeLOD::LOD1:
+		LODHashMap = ColorPointHashMapLOD1;
+		break;
+	case ESoundscapeLOD::LOD2:
+		LODHashMap = ColorPointHashMapLOD2;
+		break;
+	case ESoundscapeLOD::LOD3:
+		LODHashMap = ColorPointHashMapLOD3;
+	}
+
+	if (LODHashMap)
+	{
+		FVector BoxCenter = LODHashMap->GetHashCellCenterpoint(Location);
+		float BoxWidth = LODHashMap->GetHashCellWidth();
+
+		BoxResults.Key = BoxCenter;
+		BoxResults.Value = FVector(BoxWidth, BoxWidth, BoxWidth);
+	}
+
+	return BoxResults;
+}
+
+void USoundscapeSubsystem::AddActiveColorPoint(const USoundscapeColorPointComponent* SoundscapeColorPointComponent)
+{
+	if (SoundscapeColorPointComponent)
+	{
+		ActiveSoundscapeColorPointComponents.AddUnique(SoundscapeColorPointComponent);
+	}
+}
+
+void USoundscapeSubsystem::RemoveActiveColorPoint(const USoundscapeColorPointComponent* SoundscapeColorPointComponent)
+{
+	if (SoundscapeColorPointComponent)
+	{
+		ActiveSoundscapeColorPointComponents.RemoveSingleSwap(SoundscapeColorPointComponent, true);
+	}
+}
+
+void USoundscapeSubsystem::UpdateState()
+{
+	// Palettes the keep
+	TMap<USoundscapePalette*, UActiveSoundscapePalette*> ActivePalettesToKeep;
+
+	// Remove tags that don't match
+	for (const TPair<USoundscapePalette*, UActiveSoundscapePalette*> ActiveSoundscapePalettePair : ActivePalettes)
+	{
+		USoundscapePalette* SoundscapePalette = ActiveSoundscapePalettePair.Key;
+		UActiveSoundscapePalette* ActiveSoundscapePalette = ActiveSoundscapePalettePair.Value;
+
+		if (SoundscapePalette && ActiveSoundscapePalette)
+		{
+			if (SoundscapePalette->SoundscapePalettePlaybackConditions.Matches(SubsystemState))
+			{
+				ActivePalettesToKeep.Add(SoundscapePalette, ActiveSoundscapePalette);
+			}
+			else
+			{
+				ActiveSoundscapePalette->Stop();
+
+				UE_LOG(LogSoundscapeSubsystem, Verbose, TEXT("Soundscape Palette %s stopped"), *ActiveSoundscapePalette->GetFullName());
+			}
+		}
+	}
+
+	// Empty active palettes an append ones to keep
+	ActivePalettes.Empty();
+	ActivePalettes.Append(ActivePalettesToKeep);
+
+	// Add matching palettes
+	for (auto PaletteIterator = LoadedPaletteCollectionSet.CreateConstIterator(); PaletteIterator; ++PaletteIterator)
+	{
+		if (USoundscapePalette* SoundscapePalette = *PaletteIterator)
+		{
+			// Evaluate if the Soundscape Palette matches
+			if (SoundscapePalette->SoundscapePalettePlaybackConditions.Matches(SubsystemState) && ActivePalettes.Contains(SoundscapePalette) == false)
+			{
+				UWorld* World = GetWorld();
+
+				if (World)
+				{
+					// Create a new ActiveSoundscapeColor
+					UActiveSoundscapePalette* ActiveSoundscapePalette = NewObject<UActiveSoundscapePalette>(World);
+
+					//
+					ActiveSoundscapePalette->InitializeSettings(World, SoundscapePalette);
+
+					ActiveSoundscapePalette->Play();
+
+					UE_LOG(LogSoundscapeSubsystem, Verbose, TEXT("Soundscape Palette %s started"), *ActiveSoundscapePalette->GetFullName());
+
+					ActivePalettes.Add(SoundscapePalette, ActiveSoundscapePalette);
+
+				}
+			}
+		}
+	}
+}
+
+void USoundscapeSubsystem::UpdateColorPointHashMap(USoundscapeColorPointHashMap& SoundscapeColorPointHashMap)
+{
+	// Clear out hash Map
+	SoundscapeColorPointHashMap.ClearHash();
+
+	// Temp Collection Array
+	TArray<FSoundscapeColorPointCollection> ColorPointCollectionArray;
+
+	// Fill Temp Array
+	ColorPointCollections.GenerateValueArray(ColorPointCollectionArray);
+	
+	// Cycle through Collections
+	for (auto Collection = ColorPointCollectionArray.CreateConstIterator(); Collection; ++Collection)
+	{
+		// Cycle through Color Point Vector Arrays
+		for (auto ColorPointVectorArray = Collection->ColorPointCollection.CreateConstIterator(); ColorPointVectorArray; ++ColorPointVectorArray)
+		{
+			// Add Vector Array to Color Point Map
+			SoundscapeColorPointHashMap.AddColorPointArrayToHash(ColorPointVectorArray->Locations, ColorPointVectorArray->ColorPoint);
+		}
+	}
+}
+
+void USoundscapeSubsystem::UpdateActiveColorPointHashMap()
+{
+	ActiveColorPointHashMap->ClearHash();
+
+	for (auto It = ActiveSoundscapeColorPointComponents.CreateConstIterator(); It; ++It)
+	{
+		if (It)
+		{
+			if (const USoundscapeColorPointComponent* ActiveColorPointComponent = *It)
+			{
+				FVector ColorPointLocation;
+				FGameplayTag ColorPointValue;
+
+				ActiveColorPointComponent->GetInfo(ColorPointValue, ColorPointLocation);
+
+				ActiveColorPointHashMap->AddColorPointToHash(ColorPointLocation, ColorPointValue);
+			}
+		}
+	}
+
+	// Set up for first Timer delay
+	float DelayTime = FMath::Max(0.2f, ActiveColorPointHashUpdateTimeSeconds);
+
+	if (UWorld* World = GetWorld())
+	{
+
+		World->GetTimerManager().SetTimer(ActiveColorPointUpdateTimer, this, &USoundscapeSubsystem::UpdateActiveColorPointHashMap, DelayTime);
+	}
+}
+
+DEFINE_LOG_CATEGORY(LogSoundscapeSubsystem);
+
diff --git a/Engine/Plugins/Runtime/Soundscape/Source/SoundScape/Public/SoundScape.h b/Engine/Plugins/Runtime/Soundscape/Source/SoundScape/Public/SoundScape.h
new file mode 100644
--- /dev/null
+++ b/Engine/Plugins/Runtime/Soundscape/Source/SoundScape/Public/SoundScape.h
@@ -0,0 +1,25 @@
+// Copyright Epic Games, Inc. All Rights Reserved.
+
+#pragma once
+
+#include "CoreMinimal.h"
+#include "Kismet/BlueprintFunctionLibrary.h"
+#include "SoundScapeModule.h"
+#include "SoundscapeColor.h"
+#include "SoundScapePalette.h"
+#include "SoundScape.generated.h"
+
+UCLASS()
+class SOUNDSCAPE_API USoundscapeBPFunctionLibrary : public UBlueprintFunctionLibrary
+{
+	GENERATED_BODY()
+
+	/** Spawns a new Soundscape Elemental Agent. Returns true if spawn was successful. */
+	UFUNCTION(BlueprintCallable, Category = "Soundscape", meta = (WorldContext = "WorldContextObject"))
+	static bool SpawnSoundscapeColor(UObject* WorldContextObject, class USoundscapeColor* SoundscapeColorIn, UActiveSoundscapeColor*& ActiveSoundscapeColor);
+
+	/** Spawns a new Soundscape Palette Agent. Returns true if spawn was successful. */
+	UFUNCTION(BlueprintCallable, Category = "Soundscape", meta = (WorldContext = "WorldContextObject"))
+	static bool SpawnSoundscapePalette(UObject* WorldContextObject, class USoundscapePalette* SoundscapePaletteIn, UActiveSoundscapePalette*& SoundscapePaletteAgentsOut);
+
+};
diff --git a/Engine/Plugins/Runtime/Soundscape/Source/SoundScape/Public/SoundScapeModule.h b/Engine/Plugins/Runtime/Soundscape/Source/SoundScape/Public/SoundScapeModule.h
new file mode 100644
--- /dev/null
+++ b/Engine/Plugins/Runtime/Soundscape/Source/SoundScape/Public/SoundScapeModule.h
@@ -0,0 +1,20 @@
+// Copyright Epic Games, Inc. All Rights Reserved.
+
+#pragma once
+
+#include "Modules/ModuleInterface.h"
+#include "Modules/ModuleManager.h"
+
+DECLARE_LOG_CATEGORY_EXTERN(LogSoundscape, Log, All);
+
+class SOUNDSCAPE_API FSoundscapeModule: public IModuleInterface
+{
+public:
+
+	// IModuleInterface interface
+	virtual void StartupModule() override;
+	virtual void ShutdownModule() override;
+};
+
+
+
diff --git a/Engine/Plugins/Runtime/Soundscape/Source/SoundScape/Public/SoundScapePalette.h b/Engine/Plugins/Runtime/Soundscape/Source/SoundScape/Public/SoundScapePalette.h
new file mode 100644
--- /dev/null
+++ b/Engine/Plugins/Runtime/Soundscape/Source/SoundScape/Public/SoundScapePalette.h
@@ -0,0 +1,85 @@
+// Copyright Epic Games, Inc. All Rights Reserved.
+
+#pragma once
+
+#include "CoreMinimal.h"
+#include "SoundScapeModule.h"
+#include "Engine/UserDefinedEnum.h"
+#include "GameplayTagContainer.h"
+#include "SoundScapePalette.generated.h"
+
+class USoundscapeColor;
+class UActiveSoundscapeColor;
+
+// Struct storing Modulation State
+USTRUCT(BlueprintType)
+struct SOUNDSCAPE_API FSoundscapePaletteColor
+{
+	GENERATED_BODY()
+
+	// Soundscape Color to Play
+	UPROPERTY(EditAnywhere, BlueprintReadWrite, Category = "Soundscape|Palette")
+	USoundscapeColor* SoundscapeColor = nullptr;
+
+	// Base Volume Scalar
+	UPROPERTY(EditAnywhere, BlueprintReadWrite, Category = "Soundscape|Palette", meta = (ClampMin = "0.0", ClampMax = "4.0", UIMin = "0.0", UIMax = "4.0", SliderExponent = "6.0"))
+	float ColorVolume = 1.0f;
+
+	// Base Pitch Scalar
+	UPROPERTY(EditAnywhere, BlueprintReadWrite, Category = "Soundscape|Palette", meta = (ClampMin = "0.2", ClampMax = "4.0", UIMin = "0.2", UIMax = "4.0", SliderExponent = "3.0"))
+	float ColorPitch = 1.0f;
+
+	// Base Volume Scalar
+	UPROPERTY(EditAnywhere, BlueprintReadWrite, Category = "Soundscape|Palette", meta = (ClampMin = "0.0", UIMin = "0.0"))
+	float ColorFadeIn = 1.0f;
+
+	// Base Volume Scalar
+	UPROPERTY(EditAnywhere, BlueprintReadWrite, Category = "Soundscape|Palette", meta = (ClampMin = "0.0", UIMin = "0.0"))
+	float ColorFadeOut = 1.0f;
+};
+
+// Class containing relevant data for a Soundscape Element
+UCLASS(BlueprintType, ClassGroup = Soundscape)
+class SOUNDSCAPE_API USoundscapePalette : public UObject
+{
+	GENERATED_BODY()
+
+public:
+	USoundscapePalette();
+
+	UPROPERTY(EditAnywhere, BlueprintReadWrite, Category = "Soundscape|Palette")
+	FGameplayTagQuery SoundscapePalettePlaybackConditions;
+
+	// Elements
+	UPROPERTY(EditAnywhere, BlueprintReadWrite, Category = "Soundscape|Palette")
+	TArray<FSoundscapePaletteColor> Colors;
+
+	//~ Begin UObject Interface
+	virtual void PostLoad() override;
+	virtual void Serialize(FArchive& Ar) override;
+	//~ End UObject INterface
+};
+
+
+UCLASS(BlueprintType, ClassGroup = Soundscape)
+class SOUNDSCAPE_API UActiveSoundscapePalette : public UObject
+{
+	GENERATED_BODY()
+
+public:
+	void InitializeSettings(UObject* WorldContextObject, USoundscapePalette* SoundscapePalette);
+
+	UFUNCTION(BlueprintCallable, Category = "Soundscape")
+	void Play();
+
+	UFUNCTION(BlueprintCallable, Category = "Soundscape")
+	void Stop();
+
+private:
+	UPROPERTY()
+	UWorld* World;
+
+	UPROPERTY()
+	TArray<UActiveSoundscapeColor*> ActiveSoundscapeColors;
+
+};
diff --git a/Engine/Plugins/Runtime/Soundscape/Source/SoundScape/Public/SoundscapeColor.h b/Engine/Plugins/Runtime/Soundscape/Source/SoundScape/Public/SoundscapeColor.h
new file mode 100644
--- /dev/null
+++ b/Engine/Plugins/Runtime/Soundscape/Source/SoundScape/Public/SoundscapeColor.h
@@ -0,0 +1,456 @@
+// Copyright Epic Games, Inc. All Rights Reserved.
+
+#pragma once
+
+#include "CoreMinimal.h"
+#include "Engine/EngineTypes.h"
+#include "Engine/UserDefinedEnum.h"
+#include "Engine/World.h"
+#include "UObject/NoExportTypes.h"
+#include "GameplayTagContainer.h"
+#include "SoundscapeColor.generated.h"
+
+class USoundBase;
+class UAudioComponent;
+
+#if WITH_EDITOR
+/** UObject delegate to broadcast parameter changes to ActiveSoundscapeColor instances. */
+DECLARE_DYNAMIC_MULTICAST_DELEGATE_OneParam(FOnSoundscapeColorParameterChanges, const class USoundscapeColor*, SoundscapeColor);
+#endif
+
+UENUM(BlueprintType)
+enum class ESoundscapeColorAltitudeClampMode : uint8
+{
+	Relative = 0 UMETA(DisplayName = "Relative Altitude"), // Altitude relative to the listener location
+	Absolute = 1 UMETA(DisplayName = "Absolute Altitude") // Altitude relative to the world coordinates
+};
+
+// Struct 
+USTRUCT(BlueprintType)
+struct SOUNDSCAPE_API FSoundscapeColorSpawnSettings
+{
+	GENERATED_BODY()
+
+	// Automatically Element Spawn when Element Returns Finished
+	UPROPERTY(EditAnywhere, BlueprintReadWrite, Category = "SpawnBehavior")
+	bool bDelayFirstSpawn = false;
+
+	// Sets Minimum Time between Spawning Elements, checking if Elements are Finished, etc.
+	UPROPERTY(EditAnywhere, BlueprintReadWrite, Category = "SpawnBehavior", meta = (ClampMin = "0.0", UIMin = "0.0", EditCondition = "bDelayFirstSpawn"))
+	float MinFirstSpawnDelay = 1.0f;
+
+	// Sets the Maximum Time between Spawning Elements, checking if Elements are Finished, etc.
+	UPROPERTY(EditAnywhere, BlueprintReadWrite, Category = "SpawnBehavior", meta = (ClampMin = "0.0", UIMin = "0.0", EditCondition = "bDelayFirstSpawn"))
+	float MaxFirstSpawnDelay = 3.0f;
+
+	// When false, there will only be one Element spawned
+	UPROPERTY(EditAnywhere, BlueprintReadWrite, Category = "SpawnBehavior")
+	bool bContinuouslyRespawn = false;
+
+	// Sets Minimum Time between Spawning Elements, checking if Elements are Finished, etc.
+	UPROPERTY(EditAnywhere, BlueprintReadWrite, Category = "SpawnBehavior", meta = (ClampMin = "0.0", UIMin = "0.0", EditCondition = "bContinuouslyRespawn"))
+	float MinSpawnDelay = 1.0f;
+
+	// Sets the Maximum Time between Spawning Elements, checking if Elements are Finished, etc.
+	UPROPERTY(EditAnywhere, BlueprintReadWrite, Category = "SpawnBehavior", meta = (ClampMin = "0.0", UIMin = "0.0", EditCondition = "bContinuouslyRespawn"))
+	float MaxSpawnDelay = 3.0f;
+
+	// Max number of concurrent Elements of this type playing back at once, will not Spawn New ones until current ones are Finished Playing
+	UPROPERTY(EditAnywhere, BlueprintReadWrite, Category = "SpawnBehavior", meta = (ClampMin = "1", UIMin = "1"))
+	int32 MaxNumberOfSpawnedElements = 1;
+
+	// Sets Minimum Distance from Base for Playback of the Element
+	UPROPERTY(EditAnywhere, BlueprintReadWrite, Category = "SpawnBehavior", meta = (ClampMin = "0.0", UIMin = "0.0"))
+	float MinSpawnDistance = 0.0f;
+
+	// Sets the Maximum Distance from Base for Playback of the Element
+	UPROPERTY(EditAnywhere, BlueprintReadWrite, Category = "SpawnBehavior", meta = (ClampMin = "0.0", UIMin = "0.0"))
+	float MaxSpawnDistance = 3500.0f;
+
+	// Culls the voices by the Max Spawn Distance
+	UPROPERTY(EditAnywhere, BlueprintReadWrite, Category = "SpawnBehavior")
+	bool bCullByMaxSpawnDistance = true;
+
+	// The Minimum Azimuth based on the Base's Forward Vector in which to position Element Playback
+	UPROPERTY(EditAnywhere, BlueprintReadWrite, Category = "SpawnBehavior", meta = (ClampMin = "0.0", ClampMax = "180.0", UIMin = "0.0", UIMax = "180.0"))
+	float MinSpawnAngle = 0.0f;
+
+	// The Maximum Azimuth based on the Base's Forward Vector in which to position Element Playback
+	UPROPERTY(EditAnywhere, BlueprintReadWrite, Category = "SpawnBehavior", meta = (ClampMin = "0.0", ClampMax = "180.0", UIMin = "0.0", UIMax = "180.0"))
+	float MaxSpawnAngle = 180.0f;
+
+	// When false, there will only be one Element spawned
+	UPROPERTY(EditAnywhere, BlueprintReadWrite, Category = "SpawnBehavior")
+	bool bClampHeight = false;
+
+	// What mode of altitude clamping to use, Relative clamps the z-axis about the location of the Listener, Absolute clamps the z-axis of the world coordinates
+	UPROPERTY(EditAnywhere, BlueprintReadWrite, Category = "SpawnBehavior", meta = (EditCondition = "bClampHeight"))
+	ESoundscapeColorAltitudeClampMode ClampMode = ESoundscapeColorAltitudeClampMode::Relative;
+
+	// After Element Playback position is determined, will clamp the Z-Axis to these ranges
+	UPROPERTY(EditAnywhere, BlueprintReadWrite, Category = "SpawnBehavior", meta = (EditCondition = "bClampHeight"))
+	float MinSpawnHeightClamp = 0.0f;
+
+	// After Element Playback position is determined, will clamp the Z-Axis to these ranges
+	UPROPERTY(EditAnywhere, BlueprintReadWrite, Category = "SpawnBehavior", meta = (EditCondition = "bClampHeight"))
+	float MaxSpawnHeightClamp = 5000.0f;
+
+	// When false, the sound rotation will be unmodified
+	UPROPERTY(EditAnywhere, BlueprintReadWrite, Category = "SpawnBehavior")
+	bool bRotateSoundSource = false;
+
+	// The Minimum azimuthal rotation in the world orientation for Element Playback
+	UPROPERTY(EditAnywhere, BlueprintReadWrite, Category = "SpawnBehavior", meta = (ClampMin = "0.0", ClampMax = "360.0", UIMin = "0.0", UIMax = "360.0", EditCondition = "bRotateSoundSource"))
+	float MinAzimuthalRotationAngle = 0.0f;
+
+	// The Maximum azimuthal rotation in the world orientation for Element Playback
+	UPROPERTY(EditAnywhere, BlueprintReadWrite, Category = "SpawnBehavior", meta = (ClampMin = "0.0", ClampMax = "360.0", UIMin = "0.0", UIMax = "360.0", EditCondition = "bRotateSoundSource"))
+	float MaxAzimuthalRotationAngle = 360.0f;
+
+	// The Minimum Altitudinal rotation in the world orientation for Element Playback
+	UPROPERTY(EditAnywhere, BlueprintReadWrite, Category = "SpawnBehavior", meta = (ClampMin = "0.0", ClampMax = "360.0", UIMin = "0.0", UIMax = "360.0", EditCondition = "bRotateSoundSource"))
+	float MinAltitudinalRotationAngle = 0.0f;
+
+	// The Maximum Altitudinal rotation in the world orientation for Element Playback
+	UPROPERTY(EditAnywhere, BlueprintReadWrite, Category = "SpawnBehavior", meta = (ClampMin = "0.0", ClampMax = "360.0", UIMin = "0.0", UIMax = "360.0", EditCondition = "bRotateSoundSource"))
+	float MaxAltitudinalRotationAngle = 360.0f;
+
+	// Position By Trace allows sounds to be repositioned from their spawn point to a new point in space based on the success of an async trace call,
+	// unsuccessful hits will result in the sound not spawning
+	UPROPERTY(EditAnywhere, BlueprintReadWrite, AdvancedDisplay, Category = "SpawnBehavior")
+	bool bPositionByTrace = false;
+
+	// The Minimum nomralized directional vector in world orientation for the Trace to fire (default to down, z = -1.0f)
+	UPROPERTY(EditAnywhere, BlueprintReadWrite, AdvancedDisplay, Category = "SpawnBehavior", meta = (EditCondition = "bPositionByTrace"))
+	FVector TraceVector = { 0.0f, 0.0f, -1.0f };
+
+	// Trace distance
+	UPROPERTY(EditAnywhere, BlueprintReadWrite, AdvancedDisplay, Category = "SpawnBehavior", meta = (EditCondition = "bPositionByTrace"))
+	float TraceDistance = 500.0f;
+
+	// The physics trace channel to use
+	UPROPERTY(EditAnywhere, BlueprintReadWrite, AdvancedDisplay, Category = "SpawnBehavior", meta = (EditCondition = "bPositionByTrace"))
+	TEnumAsByte<enum ECollisionChannel> TraceChannel = ECollisionChannel::ECC_Visibility;
+
+	// Filter by Color Point Density allows designers to only allow this Color to spawn in spaces with a specified range of Color Points
+	UPROPERTY(EditAnywhere, BlueprintReadWrite, AdvancedDisplay, Category = "SpawnBehavior")
+	bool bFilterByColorPointDensity = false;
+
+	// Color Point
+	UPROPERTY(EditAnywhere, BlueprintReadWrite, AdvancedDisplay, Category = "SpawnBehavior", meta = (EditCondition = "bFilterByColorPointDensity"))
+	FGameplayTag ColorPoint;
+
+	// Minimum number of Color Point density required before spawn is successful
+	UPROPERTY(EditAnywhere, BlueprintReadWrite, AdvancedDisplay, Category = "SpawnBehavior", meta = (EditCondition = "bFilterByColorPointDensity"))
+	int32 MinColorPointNumber = 1;
+
+	// Culls the voices by the Max Spawn Distance
+	UPROPERTY(EditAnywhere, BlueprintReadWrite, AdvancedDisplay, Category = "SpawnBehavior", meta = (EditCondition = "bCullByMaxSpawnDistance"))
+	bool bIgnoreZAxisWhenCullingByMaxDistance = true;
+
+};
+
+// Struct 
+USTRUCT(BlueprintType)
+struct SOUNDSCAPE_API FSoundscapeColorModulationSettings
+{
+	GENERATED_BODY()
+
+	// Add Volume Fades to the Re-trigger Sound playback
+	UPROPERTY(EditAnywhere, BlueprintReadWrite, Category = "ModulationBehavior")
+	bool bRandomizeVolume = false;
+
+	// Minimum Volume Range for Step
+	UPROPERTY(EditAnywhere, BlueprintReadWrite, Category = "ModulationBehavior", meta = (ClampMin = "0.0", ClampMax = "4.0", UIMin = "0.0", UIMax = "4.0", SliderExponent = "6.0", EditCondition = "bRandomizeVolume"))
+	float VolumeMin = 0.95f;
+
+	// Maximum Volume Range for Step
+	UPROPERTY(EditAnywhere, BlueprintReadWrite, Category = "ModulationBehavior", meta = (ClampMin = "0.0", ClampMax = "4.0", UIMin = "0.0", UIMax = "4.0", SliderExponent = "6.0", EditCondition = "bRandomizeVolume"))
+	float VolumeMax = 1.05f;
+
+	// Add Volume Fades to the Re-trigger Sound playback
+	UPROPERTY(EditAnywhere, BlueprintReadWrite, Category = "ModulationBehavior")
+	bool bRandomizePitch = false;
+
+	// Minimum Pitch Range for Step
+	UPROPERTY(EditAnywhere, BlueprintReadWrite, Category = "ModulationBehavior", meta = (ClampMin = "0.2", ClampMax = "4.0", UIMin = "0.2", UIMax = "4.0", SliderExponent = "3.0", EditCondition = "bRandomizePitch"))
+	float PitchMin = 0.95f;
+
+	// Maximum Pitch Range for Step
+	UPROPERTY(EditAnywhere, BlueprintReadWrite, Category = "ModulationBehavior", meta = (ClampMin = "0.2", ClampMax = "4.0", UIMin = "0.2", UIMax = "4.0", SliderExponent = "3.0", EditCondition = "bRandomizePitch"))
+	float PitchMax = 1.05f;
+
+	// Add Volume Fades to the Re-trigger Sound playback
+	UPROPERTY(EditAnywhere, BlueprintReadWrite, Category = "ModulationBehavior")
+	bool bFadeVolume = false;
+
+	// This is the minimum Fade-In duration for re-triggered sounds
+	UPROPERTY(EditAnywhere, BlueprintReadWrite, Category = "ModulationBehavior", meta = (ClampMin = "0.0", UIMin = "0.0", EditCondition = "bFadeVolume"))
+	float MinFadeInTime = 0.0f;
+
+	// This is the maximum Fade-In duration for re-triggered sounds
+	UPROPERTY(EditAnywhere, BlueprintReadWrite, Category = "ModulationBehavior", meta = (ClampMin = "0.0", UIMin = "0.0", EditCondition = "bFadeVolume"))
+	float MaxFadeInTime = 0.0f;
+
+	// This is the minimum Fade-Out duration for re-triggered sounds
+	UPROPERTY(EditAnywhere, BlueprintReadWrite, Category = "ModulationBehavior", meta = (ClampMin = "0.0", UIMin = "0.0", EditCondition = "bFadeVolume"))
+	float MinFadeOutTime = 0.0f;
+
+	// This is the maximum Fade-Out duration for re-triggered sounds
+	UPROPERTY(EditAnywhere, BlueprintReadWrite, Category = "ModulationBehavior", meta = (ClampMin = "0.0", UIMin = "0.0", EditCondition = "bFadeVolume"))
+	float MaxFadeOutTime = 0.0f;
+
+	// Only Fade-In on retriggers
+	UPROPERTY(EditAnywhere, BlueprintReadWrite, Category = "ModulationBehavior", meta = (EditCondition = "bFadeVolume"))
+	bool bOnlyFadeInOnRetrigger = false;
+
+};
+
+// Struct containing Sound Scape Element Control Data
+USTRUCT(BlueprintType)
+struct SOUNDSCAPE_API FSoundscapeColorPlaybackSettings
+{
+	GENERATED_BODY()
+
+	// Upon starting playback, will randomly seek into the Sound if true
+	UPROPERTY(EditAnywhere, BlueprintReadWrite, Category = "PlaybackBehavior")
+	bool bRandomizeStartingSeekTime = false;
+
+	// Limit Ambient Bed's playback duration
+	UPROPERTY(EditAnywhere, BlueprintReadWrite, Category = "PlaybackBehavior")
+	bool bLimitPlaybackDuration = false;
+
+	// This is the minimum playback duration for this ambient bed (assuming a looping sound), ignored if bLimitPlaybackDuration is false
+	UPROPERTY(EditAnywhere, BlueprintReadWrite, Category = "PlaybackBehavior", meta = (ClampMin = "0.0", UIMin = "0.0", EditCondition = "bLimitPlaybackDuration"))
+	float MinPlaybackDuration = 0.0f;
+
+	// This is the maximum playback duration for this ambient bed (assuming a looping sound), ignored if bLimitPlaybackDuration is false
+	UPROPERTY(EditAnywhere, BlueprintReadWrite, Category = "PlaybackBehavior", meta = (ClampMin = "0.0", UIMin = "0.0", EditCondition = "bLimitPlaybackDuration"))
+	float MaxPlaybackDuration = 0.0f;
+
+};
+
+// New SoundColorVoice Spawning Parameters
+USTRUCT()
+struct SOUNDSCAPE_API FSoundscapeColorVoiceParams
+{
+	GENERATED_BODY()
+
+	float Volume = 1.0f;
+	float Pitch = 1.0f;
+	float FadeInTime = 0.0f;
+	float StartTime = 0.0f;
+	FVector Location;
+	FRotator Rotation;
+};
+
+/**
+ * 
+ */
+UCLASS(BlueprintType, ClassGroup = Soundscape)
+class SOUNDSCAPE_API USoundscapeColor : public UObject
+{
+	GENERATED_BODY()
+	
+public:
+	USoundscapeColor();
+
+	//~ Begin UObject Interface
+	virtual void PostLoad() override;
+	virtual void Serialize(FArchive& Ar) override;
+
+#if WITH_EDITORONLY_DATA
+	virtual void PostEditChangeProperty(FPropertyChangedEvent& PropertyChangedEvent) override;
+#endif
+	//~ End UObject Interface
+
+	// Sound to Play
+	UPROPERTY(EditAnywhere, BlueprintReadWrite, Category = "Soundscape|Color")
+	USoundBase* Sound;
+
+	// Base Volume Scalar
+	UPROPERTY(EditAnywhere, BlueprintReadWrite, Category = "Soundscape|Color", meta = (ClampMin = "0.0", ClampMax = "4.0", UIMin = "0.0", UIMax = "4.0", SliderExponent = "6.0"))
+	float VolumeBase = 1.0f;
+
+	// Base Pitch Scalar
+	UPROPERTY(EditAnywhere, BlueprintReadWrite, Category = "Soundscape|Color", meta = (ClampMin = "0.2", ClampMax = "4.0", UIMin = "0.2", UIMax = "4.0", SliderExponent = "3.0"))
+	float PitchBase = 1.0f;
+
+	//Modulation behavior
+	UPROPERTY(EditAnywhere, BlueprintReadWrite, Category = "Soundscape|Color", meta = (ShowOnlyInnerProperties))
+	FSoundscapeColorModulationSettings ModulationBehavior;
+
+	// Playback behavior
+	UPROPERTY(EditAnywhere, BlueprintReadWrite, Category = "Soundscape|Color", meta = (ShowOnlyInnerProperties))
+	FSoundscapeColorPlaybackSettings PlaybackBehavior;
+
+	// Playback behavior
+	UPROPERTY(EditAnywhere, BlueprintReadWrite, Category = "Soundscape|Color", meta = (ShowOnlyInnerProperties))
+	FSoundscapeColorSpawnSettings SpawnBehavior;
+
+#if WITH_EDITORONLY_DATA
+	// Delegate call for property changes
+	UPROPERTY()
+	FOnSoundscapeColorParameterChanges OnSoundscapeColorParameterChanges;
+#endif
+};
+
+UCLASS()
+class SOUNDSCAPE_API UActiveSoundscapeColorVoice : public UObject
+{
+	GENERATED_BODY()
+
+public:
+	// ~UObject Interface
+	void BeginDestroy();
+	//~UObject Interface End
+
+	UPROPERTY(EditAnywhere, Category = "Soundscape|Color")
+	UAudioComponent* AudioComponent = nullptr;
+
+	UPROPERTY(EditAnywhere, Category = "Soundscape|Color")
+	FTimerHandle TimerHandle;
+
+	UPROPERTY(EditAnywhere, Category = "Soundscape|Color")
+	bool bFree = false;
+
+	// Limited duration timer
+	UFUNCTION()
+	void StopLimitedDurationVoice(float FadeOutTime);
+};
+
+UCLASS(BlueprintType, ClassGroup = Soundscape)
+class SOUNDSCAPE_API UActiveSoundscapeColor : public UObject
+{
+	GENERATED_BODY()
+
+public:
+
+	// ~UObject Interface
+	virtual void BeginDestroy() override;
+	//~UObject Interface End
+	
+	// Initializes parameters based on SoundscapeColor settings
+	void SetParameterValues(const USoundscapeColor* SoundscapeColor);
+
+
+#if WITH_EDITOR
+	// Functionality related to Settings parameter changes to feed changes in preset Asset to instances of Active Soundscape Colors
+public:
+
+	// Setup to bind to parameter change delegate
+	void BindToParameterChangeDelegate(USoundscapeColor* SoundscapeColor);
+
+	// Function bound to the Parameter change delegate 
+	UFUNCTION()
+	void OnSoundscapeColorParameterChange(const USoundscapeColor* SoundscapeColor);
+#endif
+
+#if WITH_EDITORONLY_DATA
+	// Delegates for parameter change bindings
+	FScriptDelegate SoundscapeParameterChangeDelegate;
+#endif
+
+private:
+	// Sound to Play
+	UPROPERTY()
+	USoundBase* Sound;
+
+	// Base Volume Scalar
+	float VolumeBase = 1.0f;
+
+	// Base Pitch Scalar
+	float PitchBase = 1.0f;
+
+	//Modulation behavior
+	UPROPERTY()
+	FSoundscapeColorModulationSettings ModulationBehavior;
+
+	// Playback behavior
+	UPROPERTY()
+	FSoundscapeColorPlaybackSettings PlaybackBehavior;
+
+	// Playback behavior
+	UPROPERTY()
+	FSoundscapeColorSpawnSettings SpawnBehavior;
+
+public:
+	// Playback start values
+	float VolumeMod = 1.0f;
+	float PitchMod = 1.0f;
+	float FadeInMin = 1.0f;
+
+	// Playback stop values
+	float FadeOutMin = 1.0f;
+
+	// Native play function
+	void PlayNative();
+
+	// Native stop function
+	void StopNative();
+
+	// Play Active Soundscape Color
+	UFUNCTION(BlueprintCallable, Category = "Soundscape")
+	void Play(float ColorVolume = 1.0f, float ColorPitch = 1.0f, float ColorFadeInTime = 1.0f);
+
+	// Stop Active Soundscape Color
+	UFUNCTION(BlueprintCallable, Category = "Soundscape")
+	void Stop(float ColorFadeOutTime = 1.0f);
+
+	// Is playing
+	UFUNCTION(BlueprintCallable, Category = "Soundscape")
+	bool IsPlaying();
+
+	// Timer Update Call
+	UFUNCTION()
+	void UpdateSoundscapeColor();
+
+	// Timer delegate for amortized updates
+	FTimerHandle UpdateTimer;
+
+protected:
+
+	// Call an async trace
+	FTraceHandle SpawnSoundByTrace();
+
+	// Function called on completion of trace
+	void OnTraceCompleted(const FTraceHandle& Handle, FTraceDatum& Data);
+
+	FTraceHandle TraceHandle;
+	FTraceDelegate TraceDelegate;
+
+private:
+	// Internal begin playing and bind FTimer Delegate set first time delays
+	void StartPlaying();
+
+	// Internal stop playing, tear down FTimer delegate, etc.
+	void StopPlaying();
+
+	// Internal update call
+	void Update();
+
+	// Calculate spawning parameters
+	void CalculateSpawnParams();
+
+	// Spawn a Color Voice
+	void SpawnSoundscapeColorVoice();
+
+	// Last calculated color voice parameters
+	FSoundscapeColorVoiceParams NewSoundParams;
+
+	// Currently mid trace
+	bool bTracing = false;
+
+	// Audio Component Pool
+	UPROPERTY()
+	TArray<UActiveSoundscapeColorVoice*> SoundscapeColorVoicePool;
+
+	// Is Playing Bool
+	bool bIsPlaying = false;
+
+	// First Spawn
+	bool bFirstSpawn = true;
+
+	// Voice in pool is free
+	bool bVoiceInPoolFree = false;
+};
\ No newline at end of file
diff --git a/Engine/Plugins/Runtime/Soundscape/Source/SoundScape/Public/SoundscapeColorPoint.h b/Engine/Plugins/Runtime/Soundscape/Source/SoundScape/Public/SoundscapeColorPoint.h
new file mode 100644
--- /dev/null
+++ b/Engine/Plugins/Runtime/Soundscape/Source/SoundScape/Public/SoundscapeColorPoint.h
@@ -0,0 +1,57 @@
+// Copyright Epic Games, Inc. All Rights Reserved.
+
+#pragma once
+
+#include "CoreMinimal.h"
+#include "Components/ActorComponent.h"
+#include "GameplayTagContainer.h"
+#include "SoundscapeColorPoint.generated.h"
+
+USTRUCT(BlueprintType)
+struct SOUNDSCAPE_API FSoundscapeColorPoint
+{
+	GENERATED_BODY()
+
+	UPROPERTY(EditAnywhere, BlueprintReadWrite, Category = "Soundscape|ColorPoint")
+	FGameplayTag ColorPoint;
+
+	UPROPERTY(EditAnywhere, BlueprintReadWrite, Category = "Soundscape|ColorPoint")
+	FVector Location = FVector::ZeroVector;
+};
+
+USTRUCT(BlueprintType)
+struct SOUNDSCAPE_API FSoundscapeColorPointVectorArray
+{
+	GENERATED_BODY()
+
+	UPROPERTY(EditAnywhere, BlueprintReadWrite, Category = "Soundscape|ColorPoint")
+	FGameplayTag ColorPoint;
+
+	UPROPERTY(EditAnywhere, BlueprintReadWrite, Category = "Soundscape|ColorPoint")
+	TArray<FVector> Locations;
+};
+
+UCLASS( ClassGroup=(Soundscape), meta=(BlueprintSpawnableComponent) )
+class SOUNDSCAPE_API USoundscapeColorPointComponent : public UActorComponent
+{
+	GENERATED_BODY()
+
+public:	
+	// Sets default values for this component's properties
+	USoundscapeColorPointComponent();
+
+	// This Component's Soundscape Pigment value
+	UPROPERTY(EditAnywhere, BlueprintReadWrite, Category = "Soundscape|ColorPoint", meta = (Categories = "Soundscape"))
+	FGameplayTag ColorPoint;
+
+	// Outputs information about this Active Color Point (Color Point value and Location of Parent Actor)
+	void GetInfo(FGameplayTag& ColorPoint, FVector& Location) const;
+
+protected:
+	// Called when the game starts
+	virtual void BeginPlay() override;
+
+	// Called when the game ends
+	virtual void EndPlay(const EEndPlayReason::Type EndPlayReason) override;
+
+};
diff --git a/Engine/Plugins/Runtime/Soundscape/Source/SoundScape/Public/SoundscapeSettings.h b/Engine/Plugins/Runtime/Soundscape/Source/SoundScape/Public/SoundscapeSettings.h
new file mode 100644
--- /dev/null
+++ b/Engine/Plugins/Runtime/Soundscape/Source/SoundScape/Public/SoundscapeSettings.h
@@ -0,0 +1,68 @@
+// Copyright Epic Games, Inc. All Rights Reserved.
+
+#pragma once
+
+#include "CoreMinimal.h"
+#include "Engine/DeveloperSettings.h"
+#include "UObject/SoftObjectPath.h"
+#include "SoundscapeSettings.generated.h"
+
+/**
+ * 
+ */
+UCLASS(config = Game, defaultconfig, meta = (DisplayName = "Soundscape"))
+class SOUNDSCAPE_API USoundscapeSettings : public UDeveloperSettings
+{
+	GENERATED_UCLASS_BODY()
+
+public:
+
+	// Soundscape Palette Collection
+	UPROPERTY(config, EditAnywhere, Category = "Soundscape", meta = (AllowedClasses = "SoundscapePalette"))
+	TSet<FSoftObjectPath> SoundscapePaletteCollection;
+
+	UPROPERTY(config, EditAnywhere, Category = "Soundscape")
+	bool bDebugDraw = false;
+
+	// Hash Cell Width for LOD1
+	UPROPERTY(config, EditAnywhere, Category = "ColorPointSystem", meta = (ClampMin = "1.0", UIMin = "1.0"))
+	float LOD1ColorPointHashWidth = 500.0f;
+
+	// Hash Cell LOD1 Max Distance
+	UPROPERTY(config, EditAnywhere, Category = "ColorPointSystem", meta = (ClampMin = "1.0", UIMin = "1.0"))
+	float LOD1ColorPointHashDistance = 5000.0f;
+
+	// Hash Cell Width for LOD2
+	UPROPERTY(config, EditAnywhere, Category = "ColorPointSystem", meta = (ClampMin = "1.0", UIMin = "1.0"))
+	float LOD2ColorPointHashWidth = 2500.0f;
+
+	// Hash Cell LOD2 Max Distance
+	UPROPERTY(config, EditAnywhere, Category = "ColorPointSystem", meta = (ClampMin = "1.0", UIMin = "1.0"))
+	float LOD2ColorPointHashDistance = 10000.0f;
+
+	// Hash Cell Width for LOD3
+	UPROPERTY(config, EditAnywhere, Category = "ColorPointSystem", meta = (ClampMin = "1.0", UIMin = "1.0"))
+	float LOD3ColorPointHashWidth = 10000.0f;
+
+	// Hash Cell Width for the Active Hash
+	UPROPERTY(config, EditAnywhere, Category = "ColorPointSystem", meta = (ClampMin = "1.0", UIMin = "1.0"))
+	float ActiveColorPointHashWidth = 500.0f;
+
+	// Hash Cell Update Timing for the Active Hash
+	UPROPERTY(config, EditAnywhere, Category = "ColorPointSystem", meta = (ClampMin = "0.1", UIMin = "0.1"))
+	float ActiveColorPointHashUpdateTimeSeconds = 1.0f;
+
+public:
+
+	// Beginning of UDeveloperSettings Interface
+	virtual FName GetCategoryName() const override { return FName(TEXT("Game")); }
+#if WITH_EDITOR
+	virtual FText GetSectionText() const override { return NSLOCTEXT("SoundscapePlugin", "SoundscapeSettingsSection", "Soundscape"); };
+#endif
+	// End of UDeveloperSettings Interface
+
+
+#if WITH_EDITOR
+	virtual void PostEditChangeProperty(FPropertyChangedEvent& PropertyChangedEvent) override;
+#endif
+};
diff --git a/Engine/Plugins/Runtime/Soundscape/Source/SoundScape/Public/SoundscapeSubsystem.h b/Engine/Plugins/Runtime/Soundscape/Source/SoundScape/Public/SoundscapeSubsystem.h
new file mode 100644
--- /dev/null
+++ b/Engine/Plugins/Runtime/Soundscape/Source/SoundScape/Public/SoundscapeSubsystem.h
@@ -0,0 +1,344 @@
+// Copyright Epic Games, Inc. All Rights Reserved.
+
+#pragma once
+
+#include "CoreMinimal.h"
+#include "Subsystems/GameInstanceSubsystem.h"
+#include "GameplayTagContainer.h"
+#include "Audio.h"
+#include "SoundscapeColorPoint.h"
+#include "SoundscapeSubsystem.generated.h"
+
+class USoundscapePalette;
+class UActiveSoundscapePalette;
+
+DECLARE_LOG_CATEGORY_EXTERN(LogSoundscapeSubsystem, Log, All);
+
+// Struct 
+USTRUCT(BlueprintType)
+struct SOUNDSCAPE_API FSoundscapePaletteCollection
+{
+	GENERATED_BODY()
+
+		// Soundscape Palette Collection
+		UPROPERTY(EditAnywhere, BlueprintReadWrite, Category = "Soundscape|Palette", meta = (AllowedClasses = "SoundscapePalette"))
+		TSet<FSoftObjectPath> SoundscapePaletteCollection;
+};
+
+// Struct 
+USTRUCT()
+struct SOUNDSCAPE_API FSoundscapePaletteCollectionLoaded
+{
+	GENERATED_BODY()
+
+	// Soundscape Palette Collection
+	UPROPERTY(EditAnywhere, Category = "Soundscape|Palette")
+	TSet<USoundscapePalette*> SoundscapePaletteCollection;
+};
+
+// Struct
+USTRUCT(BlueprintType)
+struct SOUNDSCAPE_API FSoundscapeColorPointCollection
+{
+	GENERATED_BODY()
+
+	// Soundscape Color Point Collection
+	UPROPERTY(EditAnywhere, BlueprintReadWrite, Category = "Soundscape|ColorPoint")
+	TArray<FSoundscapeColorPointVectorArray> ColorPointCollection;
+};
+
+// Class
+UCLASS()
+class SOUNDSCAPE_API USoundscapeColorPointHashCellDensity : public UObject
+{
+	GENERATED_BODY()
+
+public:
+
+	// Soundscape Color Point Density for a Hash Cell
+	UPROPERTY(EditAnywhere, Category = "Soundscape|ColorPoint")
+	TMap<uint64, int32> ColorPointHashCellDensity;
+
+};
+
+UCLASS()
+class SOUNDSCAPE_API USoundscapeColorPointHashMap : public UObject
+{
+	GENERATED_BODY()
+
+public:
+	// Clears and initializes hash map, sizes grid to HashCellWidth in Uunits
+	UFUNCTION()
+	void InitializeHash(float HashCellSizeIn, FVector GridCenterIn);
+
+	// Clears hash map
+	UFUNCTION()
+	void ClearHash();
+
+	// Returns the number of ColorPoints in a Cell containing the Location
+	UFUNCTION()
+	int32 NumColorPointsInCell(const FVector& Location, const FGameplayTag ColorPoint);
+
+	// Returns true if ColorPoint added to hash, false if failed (likely location is out of Hash Bounds)
+	UFUNCTION()
+	bool AddColorPointToHash(const FVector& Location, const FGameplayTag ColorPoint);
+
+	// Returns true if ColorPoint added to hash, false if failed (likely location is out of Hash Bounds)
+	UFUNCTION()
+	void AddColorPointArrayToHash(const TArray<FVector>& Locations, const FGameplayTag ColorPoint);
+
+	// Calculates Hash Index
+	UFUNCTION()
+	uint64 CalculateHashIndex(const FVector& Location);
+
+	// Sets ceterpoint of grid to Location
+	UFUNCTION()
+	void SetGridCenterpoint(const FVector& Location);
+
+	TMap<FGameplayTag, int32> GetHashMapColorPointDensitySummary();
+	
+protected:
+	// Returns ColorPointDensityMap from Hash Map given a Color Point
+	USoundscapeColorPointHashCellDensity* GetColorPointDensityMap(FGameplayTag ColorPoint);
+
+private:
+	// Color Point Hash Map
+	UPROPERTY()
+	TMap<FGameplayTag, USoundscapeColorPointHashCellDensity*> ColorPointHashMap;
+
+	// Cubic Root of the Max Value of int64 rounded down to the nearest million
+	static const int32 MaxGridWidth = 2000000;
+	static const int32 HalfMaxGridWidth = 1000000;
+
+	// Hash Cell Width
+	UPROPERTY()
+	float HashCellSize = 500.0f;
+
+	// MaxGridWidth / Hash Cell Size
+	UPROPERTY()
+	int32 GridWidth = 4000;
+
+	// Grid Width Cubed
+	UPROPERTY()
+	int64 NumCells = 64000000000;
+
+	// Cell Size Inverted
+	UPROPERTY()
+	float HashCellFactor = 0.002f;
+
+	// Current Grid Centerpoint
+	UPROPERTY()
+	FVector GridCenter;
+
+	// Grid Origin Offset
+	UPROPERTY()
+	FVector GridOriginOffset;
+
+public:
+	const FVector GetHashCellCenterpoint(FVector Location);
+	const float GetHashCellWidth();
+};
+
+UENUM()
+enum class ESoundscapeLOD : uint8
+{
+	LOD1 = 0 UMETA(DisplayName = "LOD 1"),
+	LOD2 UMETA(DisplayName = "LOD 2"),
+	LOD3 UMETA(DisplayName = "LOD 3")
+};
+
+UCLASS()
+class SOUNDSCAPE_API USoundscapeColorPointHashMapCollection : public UObject
+{
+	GENERATED_BODY()
+
+public:
+
+	void InitializeCollection();
+
+	void AddColorPointArrayToHashMapCollection(const TArray<FVector>& Locations, const FGameplayTag ColorPoint);
+
+	void ClearColorPointHashMapCollection();
+
+	int32 GetColorPointHashMapCollectionDensity(const FVector Location, const FGameplayTag ColorPoint, const ESoundscapeLOD SoundscapeLOD);
+
+public:
+	void CalculateTotalColorPointDensity(TMap<FGameplayTag, int32>& TotalHashMapColorPointDensity, ESoundscapeLOD SoundscapeLOD);
+
+private:
+
+	// Color Point Hash Maps
+	UPROPERTY()
+	USoundscapeColorPointHashMap* ColorPointHashMapLOD1;
+	UPROPERTY()
+	USoundscapeColorPointHashMap* ColorPointHashMapLOD2;
+	UPROPERTY()
+	USoundscapeColorPointHashMap* ColorPointHashMapLOD3;
+
+	// Hash Cell Width for LOD1
+	UPROPERTY()
+	float LOD1ColorPointHashWidth = 500.0f;
+
+	// Hash Cell LOD1 Max Distance
+	UPROPERTY()
+	float LOD1ColorPointHashDistance = 5000.0f;
+
+	// Hash Cell Width for LOD2
+	UPROPERTY()
+	float LOD2ColorPointHashWidth = 2500.0f;
+
+	// Hash Cell LOD2 Max Distance
+	UPROPERTY()
+	float LOD2ColorPointHashDistance = 10000.0f;
+
+	// Hash Cell Width for LOD3
+	UPROPERTY()
+	float LOD3ColorPointHashWidth = 10000.0f;
+
+};
+
+/**
+ * 
+ */
+UCLASS()
+class SOUNDSCAPE_API USoundscapeSubsystem : public UGameInstanceSubsystem
+{
+	GENERATED_BODY()
+	
+public:
+	// Begin USubsystem
+	virtual void Initialize(FSubsystemCollectionBase& Collection) override;
+	virtual void Deinitialize() override;
+	virtual bool ShouldCreateSubsystem(UObject* Outer) const override;
+	// End USubsystem
+
+
+	// Settings
+public:
+
+	UFUNCTION(BlueprintCallable, Category = "Soundscape")
+	void SetState(FGameplayTag SoundscapeState);
+
+	UFUNCTION(BlueprintCallable, Category = "Soundscape")
+	void ClearState(FGameplayTag SoundscapeState);
+
+	UFUNCTION(BlueprintCallable, Category = "Soundscape")
+	void RestartSoundscape();
+
+	UFUNCTION(BlueprintCallable, Category = "Soundscape")
+	bool AddPaletteCollection(FName PaletteCollectionName, FSoundscapePaletteCollection PaletteCollection);
+
+	UFUNCTION(BlueprintCallable, Category = "Soundscape")
+	bool RemovePaletteCollection(FName PaletteCollectionName);
+
+	bool bDebugMode = false;
+
+private:
+	UPROPERTY()
+	TSet<USoundscapePalette*> LoadedPaletteCollectionSet;
+
+	UPROPERTY()
+	TMap<FName, FSoundscapePaletteCollection> UnloadedPaletteCollections;
+
+	FGameplayTagContainer SubsystemState;
+
+	UPROPERTY()
+	TMap<USoundscapePalette*, UActiveSoundscapePalette*> ActivePalettes;
+
+	bool LoadPaletteCollection(FName PaletteCollectionName);
+
+	bool UnloadPaletteCollection(FName PaletteCollectionName);
+
+	void UpdateState();
+
+private:
+
+	Audio::FDeviceId AudioDeviceID;
+
+	// Soundscape Color Point system
+public:
+	// Add a Color Point Collection to the Subsystem, returns true if successful
+	UFUNCTION(BlueprintCallable, Category = "Soundscape")
+	void AddColorPointCollection(FName ColorPointCollectionName, FSoundscapeColorPointCollection ColorPointCollection);
+
+	// Remove a Color Point Collection from the Subsystem, returns true if successful
+	UFUNCTION(BlueprintCallable, Category = "Soundscape")
+	bool RemoveColorPointCollection(FName ColorPointCollectionName);
+
+	// Check Color Point Density for a Location Cell
+	UFUNCTION(BlueprintCallable, Category = "Soundscape")
+	int32 CheckColorPointDensity(FVector Location, FGameplayTag ColorPoint);
+
+	// Add Hash Map Collection
+	void AddColorPointHashMapCollection(USoundscapeColorPointHashMapCollection* ColorPointHashMapCollection);
+
+	// Add Hash Map Collection
+	bool RemoveColorPointHashMapCollection(USoundscapeColorPointHashMapCollection* ColorPointHashMapCollection);
+
+	// Draw Debug Cell
+	void DrawDebugCell(FVector Location, bool bSuccess);
+
+	// Calculate Debug Cell Dimensions
+	TPair<FVector, FVector> CalculateDebugCellDimensions(FVector Location, ESoundscapeLOD SoundscapeLOD);
+
+	// Add Active Color Point
+	void AddActiveColorPoint(const USoundscapeColorPointComponent* SoundscapeColorPointComponent);
+
+	// Remove Active Color Point
+	void RemoveActiveColorPoint(const USoundscapeColorPointComponent* SoundscapeColorPointComponent);
+
+private:
+	// Stored Hash Map Collections
+	UPROPERTY(Transient)
+	TArray<USoundscapeColorPointHashMapCollection*> ColorPointHashMapCollections;
+
+	// Stored Color Point Collections
+	TMap<FName, FSoundscapeColorPointCollection> ColorPointCollections;
+
+	// Active Color Point Collection
+	FSoundscapeColorPointCollection ActiveColorPointCollection;
+
+	// Color Point Hash Maps
+	UPROPERTY()
+	USoundscapeColorPointHashMap* ColorPointHashMapLOD1;
+	UPROPERTY()
+	USoundscapeColorPointHashMap* ColorPointHashMapLOD2;
+	UPROPERTY()
+	USoundscapeColorPointHashMap* ColorPointHashMapLOD3;
+
+	UPROPERTY()
+	USoundscapeColorPointHashMap* ActiveColorPointHashMap;
+
+	// Hash Cell Width for LOD1
+	float LOD1ColorPointHashWidth = 500.0f;
+
+	// Hash Cell LOD1 Max Distance
+	float LOD1ColorPointHashDistance = 5000.0f;
+
+	// Hash Cell Width for LOD2
+	float LOD2ColorPointHashWidth = 2500.0f;
+
+	// Hash Cell LOD2 Max Distance
+	float LOD2ColorPointHashDistance = 10000.0f;
+
+	// Hash Cell Width for LOD3
+	float LOD3ColorPointHashWidth = 10000.0f;
+
+	// Hash Cell Width for the Active Hash
+	float ActiveColorPointHashWidth = 500.0f;
+
+	// Hash Cell Update Timing for the Active Hash
+	float ActiveColorPointHashUpdateTimeSeconds = 1.0f;
+
+	// Update Hash Map
+	void UpdateColorPointHashMap(USoundscapeColorPointHashMap& SoundscapeColorPointHashMap);
+
+	// Update Active Hash Map
+	void UpdateActiveColorPointHashMap();
+
+	// Timer delegate for amortized updates
+	FTimerHandle ActiveColorPointUpdateTimer;
+	
+	// Active Color Point Components
+	TArray<const USoundscapeColorPointComponent*> ActiveSoundscapeColorPointComponents;
+};
diff --git a/Engine/Plugins/Runtime/Soundscape/Source/SoundScape/SoundScape.Build.cs b/Engine/Plugins/Runtime/Soundscape/Source/SoundScape/SoundScape.Build.cs
new file mode 100644
--- /dev/null
+++ b/Engine/Plugins/Runtime/Soundscape/Source/SoundScape/SoundScape.Build.cs
@@ -0,0 +1,55 @@
+// Copyright Epic Games, Inc. All Rights Reserved.
+
+namespace UnrealBuildTool.Rules
+{
+	public class Soundscape : ModuleRules
+	{
+        public Soundscape(ReadOnlyTargetRules Target) : base(Target)
+        {
+			PCHUsage = ModuleRules.PCHUsageMode.UseExplicitOrSharedPCHs;
+
+			PublicIncludePaths.AddRange(
+				new string[] {
+					// ... add public include paths required here ...
+				}
+				);
+
+
+			PrivateIncludePaths.AddRange(
+				new string[] {
+					// ... add other private include paths required here ...
+				}
+				);
+
+
+			PublicDependencyModuleNames.AddRange(
+				new string[]
+				{
+				"Core",
+				"CoreUObject",
+				"DeveloperSettings",
+				"GameplayTags",
+					// ... add other public dependencies that you statically link with here ...
+				}
+				);
+
+
+			PrivateDependencyModuleNames.AddRange(
+				new string[]
+				{
+				"CoreUObject",
+				"Engine",
+					// ... add private dependencies that you statically link with here ...	
+				}
+				);
+
+
+			DynamicallyLoadedModuleNames.AddRange(
+				new string[]
+				{
+					// ... add any modules that your module loads dynamically here ...
+				}
+				);
+		}
+	}
+}
\ No newline at end of file
diff --git a/Engine/Plugins/Runtime/Soundscape/Source/SoundScapeEditor/Private/AssetTypeActions_SoundScapePalette.cpp b/Engine/Plugins/Runtime/Soundscape/Source/SoundScapeEditor/Private/AssetTypeActions_SoundScapePalette.cpp
new file mode 100644
--- /dev/null
+++ b/Engine/Plugins/Runtime/Soundscape/Source/SoundScapeEditor/Private/AssetTypeActions_SoundScapePalette.cpp
@@ -0,0 +1,23 @@
+// Copyright Epic Games, Inc. All Rights Reserved.
+
+#include "AssetTypeActions_SoundScapePalette.h"
+#include "SoundScapePalette.h"
+
+#define LOCTEXT_NAMESPACE "AssetTypeActions"
+
+UClass* FAssetTypeActions_SoundscapePalette::GetSupportedClass() const
+{
+	return USoundscapePalette::StaticClass();
+}
+
+const TArray<FText>& FAssetTypeActions_SoundscapePalette::GetSubMenus() const
+{
+	static const TArray<FText> SubMenus
+	{
+		LOCTEXT("AssetSoundscapeSubMenu", "Soundscape")
+	};
+	return SubMenus;
+}
+
+
+#undef LOCTEXT_NAMESPACE
diff --git a/Engine/Plugins/Runtime/Soundscape/Source/SoundScapeEditor/Private/AssetTypeActions_SoundScapePalette.h b/Engine/Plugins/Runtime/Soundscape/Source/SoundScapeEditor/Private/AssetTypeActions_SoundScapePalette.h
new file mode 100644
--- /dev/null
+++ b/Engine/Plugins/Runtime/Soundscape/Source/SoundScapeEditor/Private/AssetTypeActions_SoundScapePalette.h
@@ -0,0 +1,18 @@
+// Copyright Epic Games, Inc. All Rights Reserved.
+
+#pragma once
+
+#include "CoreMinimal.h"
+#include "Toolkits/IToolkitHost.h"
+#include "AssetTypeActions_Base.h"
+
+class FAssetTypeActions_SoundscapePalette : public FAssetTypeActions_Base
+{
+public:
+	// IAssetTypeActions Implementation
+	virtual FText GetName() const override { return NSLOCTEXT("AssetTypeActions", "AssetTypeActions_SoundscapePalette", "Soundscape Palette"); }
+	virtual FColor GetTypeColor() const override { return FColor(0, 125, 255); }
+	virtual const TArray<FText>& GetSubMenus() const override;
+	virtual UClass* GetSupportedClass() const override;
+	virtual uint32 GetCategories() override { return EAssetTypeCategories::Sounds; }
+};
diff --git a/Engine/Plugins/Runtime/Soundscape/Source/SoundScapeEditor/Private/AssetTypeActions_SoundscapeColor.cpp b/Engine/Plugins/Runtime/Soundscape/Source/SoundScapeEditor/Private/AssetTypeActions_SoundscapeColor.cpp
new file mode 100644
--- /dev/null
+++ b/Engine/Plugins/Runtime/Soundscape/Source/SoundScapeEditor/Private/AssetTypeActions_SoundscapeColor.cpp
@@ -0,0 +1,23 @@
+// Copyright Epic Games, Inc. All Rights Reserved.
+
+#include "AssetTypeActions_SoundscapeColor.h"
+#include "SoundscapeColor.h"
+
+#define LOCTEXT_NAMESPACE "AssetTypeActions"
+
+UClass* FAssetTypeActions_SoundscapeColor::GetSupportedClass() const
+{
+	return USoundscapeColor::StaticClass();
+}
+
+const TArray<FText>& FAssetTypeActions_SoundscapeColor::GetSubMenus() const
+{
+	static const TArray<FText> SubMenus
+	{
+		LOCTEXT("AssetSoundscapeSubMenu", "Soundscape")
+	};
+	return SubMenus;
+}
+
+
+#undef LOCTEXT_NAMESPACE
diff --git a/Engine/Plugins/Runtime/Soundscape/Source/SoundScapeEditor/Private/AssetTypeActions_SoundscapeColor.h b/Engine/Plugins/Runtime/Soundscape/Source/SoundScapeEditor/Private/AssetTypeActions_SoundscapeColor.h
new file mode 100644
--- /dev/null
+++ b/Engine/Plugins/Runtime/Soundscape/Source/SoundScapeEditor/Private/AssetTypeActions_SoundscapeColor.h
@@ -0,0 +1,18 @@
+// Copyright Epic Games, Inc. All Rights Reserved.
+
+#pragma once
+
+#include "CoreMinimal.h"
+#include "Toolkits/IToolkitHost.h"
+#include "AssetTypeActions_Base.h"
+
+class FAssetTypeActions_SoundscapeColor : public FAssetTypeActions_Base
+{
+public:
+	// IAssetTypeActions Implementation
+	virtual FText GetName() const override { return NSLOCTEXT("AssetTypeActions", "AssetTypeActions_SoundscapeColor", "Soundscape Color"); }
+	virtual FColor GetTypeColor() const override { return FColor(0, 175, 175); }
+	virtual const TArray<FText>& GetSubMenus() const override;
+	virtual UClass* GetSupportedClass() const override;
+	virtual uint32 GetCategories() override { return EAssetTypeCategories::Sounds; }
+};
diff --git a/Engine/Plugins/Runtime/Soundscape/Source/SoundScapeEditor/Private/SoundScapeEditorModule.cpp b/Engine/Plugins/Runtime/Soundscape/Source/SoundScapeEditor/Private/SoundScapeEditorModule.cpp
new file mode 100644
--- /dev/null
+++ b/Engine/Plugins/Runtime/Soundscape/Source/SoundScapeEditor/Private/SoundScapeEditorModule.cpp
@@ -0,0 +1,29 @@
+// Copyright Epic Games, Inc. All Rights Reserved.
+
+#include "SoundScapeEditorModule.h"
+#include "AssetToolsModule.h"
+#include "CoreMinimal.h"
+#include "Stats/Stats.h"
+#include "Modules/ModuleInterface.h"
+#include "Modules/ModuleManager.h"
+#include "AssetTypeActions_SoundscapeColor.h"
+#include "AssetTypeActions_SoundScapePalette.h"
+
+
+
+IMPLEMENT_MODULE(FSoundscapeEditorModule, SoundscapeEditor)
+
+void FSoundscapeEditorModule::StartupModule()
+{
+
+	// Register the audio editor asset type actions.
+	IAssetTools& AssetTools = FModuleManager::LoadModuleChecked<FAssetToolsModule>("AssetTools").Get();
+
+	AssetTools.RegisterAssetTypeActions(MakeShareable(new FAssetTypeActions_SoundscapeColor));
+	AssetTools.RegisterAssetTypeActions(MakeShareable(new FAssetTypeActions_SoundscapePalette));
+}
+
+void FSoundscapeEditorModule::ShutdownModule()
+{
+}
+
diff --git a/Engine/Plugins/Runtime/Soundscape/Source/SoundScapeEditor/Private/SoundScapePaletteFactory.cpp b/Engine/Plugins/Runtime/Soundscape/Source/SoundScapeEditor/Private/SoundScapePaletteFactory.cpp
new file mode 100644
--- /dev/null
+++ b/Engine/Plugins/Runtime/Soundscape/Source/SoundScapeEditor/Private/SoundScapePaletteFactory.cpp
@@ -0,0 +1,21 @@
+// Copyright Epic Games, Inc. All Rights Reserved.
+
+#include "SoundScapePaletteFactory.h"
+#include "SoundScapePalette.h"
+
+USoundscapePaletteFactory::USoundscapePaletteFactory(const FObjectInitializer& ObjectInitializer)
+	: Super(ObjectInitializer)
+{
+	SupportedClass = USoundscapePalette::StaticClass();
+
+	bCreateNew = true;
+	bEditorImport = false;
+	bEditAfterNew = true;
+}
+
+UObject* USoundscapePaletteFactory::FactoryCreateNew(UClass* Class, UObject* InParent, FName Name, EObjectFlags Flags, UObject* Context, FFeedbackContext* Warn)
+{
+	USoundscapePalette* SoundscapePalette = NewObject<USoundscapePalette>(InParent, Name, Flags);
+
+	return SoundscapePalette;
+}
diff --git a/Engine/Plugins/Runtime/Soundscape/Source/SoundScapeEditor/Private/SoundScapePaletteFactory.h b/Engine/Plugins/Runtime/Soundscape/Source/SoundScapeEditor/Private/SoundScapePaletteFactory.h
new file mode 100644
--- /dev/null
+++ b/Engine/Plugins/Runtime/Soundscape/Source/SoundScapeEditor/Private/SoundScapePaletteFactory.h
@@ -0,0 +1,24 @@
+// Copyright Epic Games, Inc. All Rights Reserved.
+
+#pragma once
+
+#include "CoreMinimal.h"
+#include "UObject/ObjectMacros.h"
+#include "Factories/Factory.h"
+#include "SoundScapePaletteFactory.generated.h"
+
+UCLASS(hidecategories = Object, MinimalAPI)
+class USoundscapePaletteFactory : public UFactory
+{
+	GENERATED_UCLASS_BODY()
+
+	//~ Begin UFactory Interface
+	virtual UObject* FactoryCreateNew(UClass* Class, UObject* InParent, FName Name, EObjectFlags Flags, UObject* Context, FFeedbackContext* Warn) override;
+
+	virtual bool ShouldShowInNewMenu() const override
+	{
+		return true;
+	}
+	//~ End UFactory Interface	
+};
+
diff --git a/Engine/Plugins/Runtime/Soundscape/Source/SoundScapeEditor/Private/SoundscapeColorFactory.cpp b/Engine/Plugins/Runtime/Soundscape/Source/SoundScapeEditor/Private/SoundscapeColorFactory.cpp
new file mode 100644
--- /dev/null
+++ b/Engine/Plugins/Runtime/Soundscape/Source/SoundScapeEditor/Private/SoundscapeColorFactory.cpp
@@ -0,0 +1,21 @@
+// Copyright Epic Games, Inc. All Rights Reserved.
+
+#include "SoundscapeColorFactory.h"
+#include "SoundscapeColor.h"
+
+USoundscapeColorFactory::USoundscapeColorFactory(const FObjectInitializer& ObjectInitializer)
+	: Super(ObjectInitializer)
+{
+	SupportedClass = USoundscapeColor::StaticClass();
+
+	bCreateNew = true;
+	bEditorImport = false;
+	bEditAfterNew = true;
+}
+
+UObject* USoundscapeColorFactory::FactoryCreateNew(UClass* Class, UObject* InParent, FName Name, EObjectFlags Flags, UObject* Context, FFeedbackContext* Warn)
+{
+	USoundscapeColor* SoundscapeColor = NewObject<USoundscapeColor>(InParent, Name, Flags);
+
+	return SoundscapeColor;
+}
\ No newline at end of file
diff --git a/Engine/Plugins/Runtime/Soundscape/Source/SoundScapeEditor/Private/SoundscapeColorFactory.h b/Engine/Plugins/Runtime/Soundscape/Source/SoundScapeEditor/Private/SoundscapeColorFactory.h
new file mode 100644
--- /dev/null
+++ b/Engine/Plugins/Runtime/Soundscape/Source/SoundScapeEditor/Private/SoundscapeColorFactory.h
@@ -0,0 +1,24 @@
+// Copyright Epic Games, Inc. All Rights Reserved.
+
+#pragma once
+
+#include "CoreMinimal.h"
+#include "UObject/ObjectMacros.h"
+#include "Factories/Factory.h"
+#include "SoundscapeColorFactory.generated.h"
+
+UCLASS(hidecategories = Object, MinimalAPI)
+class USoundscapeColorFactory : public UFactory
+{
+	GENERATED_UCLASS_BODY()
+
+	//~ Begin UFactory Interface
+	virtual UObject* FactoryCreateNew(UClass* Class, UObject* InParent, FName Name, EObjectFlags Flags, UObject* Context, FFeedbackContext* Warn) override;
+
+	virtual bool ShouldShowInNewMenu() const override
+	{
+		return true;
+	}
+	//~ End UFactory Interface	
+};
+
diff --git a/Engine/Plugins/Runtime/Soundscape/Source/SoundScapeEditor/Public/SoundScapeEditorModule.h b/Engine/Plugins/Runtime/Soundscape/Source/SoundScapeEditor/Public/SoundScapeEditorModule.h
new file mode 100644
--- /dev/null
+++ b/Engine/Plugins/Runtime/Soundscape/Source/SoundScapeEditor/Public/SoundScapeEditorModule.h
@@ -0,0 +1,17 @@
+// Copyright Epic Games, Inc. All Rights Reserved.
+
+#pragma once
+
+#include "Modules/ModuleInterface.h"
+#include "Modules/ModuleManager.h"
+
+
+class FSoundscapeEditorModule : public IModuleInterface
+{
+public:
+
+	// IModuleInterface interface
+	virtual void StartupModule() override;
+	virtual void ShutdownModule() override;
+
+};
diff --git a/Engine/Plugins/Runtime/Soundscape/Source/SoundScapeEditor/SoundScapeEditor.Build.cs b/Engine/Plugins/Runtime/Soundscape/Source/SoundScapeEditor/SoundScapeEditor.Build.cs
new file mode 100644
--- /dev/null
+++ b/Engine/Plugins/Runtime/Soundscape/Source/SoundScapeEditor/SoundScapeEditor.Build.cs
@@ -0,0 +1,64 @@
+// Copyright Epic Games, Inc. All Rights Reserved.
+
+namespace UnrealBuildTool.Rules
+{
+	public class SoundscapeEditor : ModuleRules
+	{
+		public SoundscapeEditor(ReadOnlyTargetRules Target) : base(Target)
+        {
+
+			PCHUsage = ModuleRules.PCHUsageMode.UseExplicitOrSharedPCHs;
+
+			PublicIncludePaths.AddRange(
+				new string[] {
+					// ... add public include paths required here ...
+				}
+				);
+
+
+			PrivateIncludePaths.AddRange(
+				new string[] {
+					// ... add other private include paths required here ...
+				}
+				);
+
+
+			PublicDependencyModuleNames.AddRange(
+				new string[]
+				{
+				"Core",
+				"AssetTools",
+				"Soundscape",
+				"AudioMixer",
+
+					// ... add other public dependencies that you statically link with here ...
+				}
+				);
+
+
+			PrivateDependencyModuleNames.AddRange(
+				new string[]
+				{
+				"CoreUObject",
+				"Engine",
+				"Slate",
+				"SlateCore",
+				"UnrealEd",
+					// ... add private dependencies that you statically link with here ...	
+				}
+				);
+
+
+			DynamicallyLoadedModuleNames.AddRange(
+				new string[]
+				{
+					// ... add any modules that your module loads dynamically here ...
+				}
+				);
+
+			PrivateIncludePathModuleNames.AddRange(
+			new string[] {
+            });
+		}
+	}
+}
\ No newline at end of file
