Fix for static analysis warning

diff --git a/Engine/Plugins/Runtime/Soundscape/Source/SoundScape/Private/SoundscapeSubsystem.cpp b/Engine/Plugins/Runtime/Soundscape/Source/SoundScape/Private/SoundscapeSubsystem.cpp
--- a/Engine/Plugins/Runtime/Soundscape/Source/SoundScape/Private/SoundscapeSubsystem.cpp
+++ b/Engine/Plugins/Runtime/Soundscape/Source/SoundScape/Private/SoundscapeSubsystem.cpp
@@ -217,7 +217,7 @@ uint64 USoundscapeColorPointHashMap::CalculateHashIndex(const FVector& Location)
 	const FVector OffsetLocation = Location - GridOriginOffset;
 	uint64 HashX = uint64(FMath::FloorToInt(OffsetLocation.X * HashCellFactor));
 	uint64 HashY = uint64(FMath::FloorToInt(OffsetLocation.Y * HashCellFactor)) * uint64(GridWidth);
-	uint64 HashZ = uint64(FMath::FloorToInt(OffsetLocation.Z * HashCellFactor)) * uint64(GridWidth * GridWidth);
+	uint64 HashZ = uint64(FMath::FloorToInt(OffsetLocation.Z * HashCellFactor)) * (uint64(GridWidth) * uint64(GridWidth));
 
 	return HashX + HashY + HashZ;
 }
