Fix Audio-related Short Type Name warnings in AncientGame

diff --git a/Engine/Plugins/Runtime/Soundscape/Source/SoundScape/Public/SoundscapeSettings.h b/Engine/Plugins/Runtime/Soundscape/Source/SoundScape/Public/SoundscapeSettings.h
--- a/Engine/Plugins/Runtime/Soundscape/Source/SoundScape/Public/SoundscapeSettings.h
+++ b/Engine/Plugins/Runtime/Soundscape/Source/SoundScape/Public/SoundscapeSettings.h
@@ -18,7 +18,7 @@ class SOUNDSCAPE_API USoundscapeSettings : public UDeveloperSettings
 public:
 
 	// Soundscape Palette Collection
-	UPROPERTY(config, EditAnywhere, Category = "Soundscape", meta = (AllowedClasses = "SoundscapePalette"))
+	UPROPERTY(config, EditAnywhere, Category = "Soundscape", meta = (AllowedClasses = "/Script/Soundscape.SoundscapePalette"))
 	TSet<FSoftObjectPath> SoundscapePaletteCollection;
 
 	UPROPERTY(config, EditAnywhere, Category = "Soundscape")
diff --git a/Engine/Plugins/Runtime/Soundscape/Source/SoundScape/Public/SoundscapeSubsystem.h b/Engine/Plugins/Runtime/Soundscape/Source/SoundScape/Public/SoundscapeSubsystem.h
--- a/Engine/Plugins/Runtime/Soundscape/Source/SoundScape/Public/SoundscapeSubsystem.h
+++ b/Engine/Plugins/Runtime/Soundscape/Source/SoundScape/Public/SoundscapeSubsystem.h
@@ -21,7 +21,7 @@ struct SOUNDSCAPE_API FSoundscapePaletteCollection
 	GENERATED_BODY()
 
 		// Soundscape Palette Collection
-		UPROPERTY(EditAnywhere, BlueprintReadWrite, Category = "Soundscape|Palette", meta = (AllowedClasses = "SoundscapePalette"))
+		UPROPERTY(EditAnywhere, BlueprintReadWrite, Category = "Soundscape|Palette", meta = (AllowedClasses = "/Script/Soundscape.SoundscapePalette"))
 		TSet<FSoftObjectPath> SoundscapePaletteCollection;
 };
 
