Pass 2 on cleaning up build.cs files.

diff --git a/Engine/Plugins/Runtime/NetworkPredictionExtras/Source/NetworkPredictionExtras/NetworkPredictionExtras.Build.cs b/Engine/Plugins/Runtime/NetworkPredictionExtras/Source/NetworkPredictionExtras/NetworkPredictionExtras.Build.cs
--- a/Engine/Plugins/Runtime/NetworkPredictionExtras/Source/NetworkPredictionExtras/NetworkPredictionExtras.Build.cs
+++ b/Engine/Plugins/Runtime/NetworkPredictionExtras/Source/NetworkPredictionExtras/NetworkPredictionExtras.Build.cs
@@ -6,18 +6,6 @@ namespace UnrealBuildTool.Rules
 	{
 		public NetworkPredictionExtras(ReadOnlyTargetRules Target) : base(Target)
 		{
-			PublicIncludePaths.AddRange(
-				new string[] {
-                    ModuleDirectory + "/Public",
-				}
-				);
-
-			PrivateIncludePaths.AddRange(
-				new string[] {
-                    "NetworkPredictionExtras/Private",
-				}
-				);
-
 			PublicDependencyModuleNames.AddRange(
 				new string[]
 				{
diff --git a/Engine/Plugins/Runtime/Oculus/OculusVR/Source/OculusOpenXRHMD/OculusOpenXRHMD.Build.cs b/Engine/Plugins/Runtime/Oculus/OculusVR/Source/OculusOpenXRHMD/OculusOpenXRHMD.Build.cs
--- a/Engine/Plugins/Runtime/Oculus/OculusVR/Source/OculusOpenXRHMD/OculusOpenXRHMD.Build.cs
+++ b/Engine/Plugins/Runtime/Oculus/OculusVR/Source/OculusOpenXRHMD/OculusOpenXRHMD.Build.cs
@@ -81,11 +81,6 @@ namespace UnrealBuildTool.Rules
             }
             else if (Target.Platform == UnrealTargetPlatform.Android)
             {
-                PrivateIncludePaths.AddRange(
-                        new string[]
-                        {
-                        });
-
                 // Vulkan
                 {
                     AddEngineThirdPartyPrivateStaticDependencies(Target, "Vulkan");
diff --git a/Engine/Plugins/Runtime/PacketHandlers/AESGCMHandlerComponent/Source/AESGCMHandlerComponent.Build.cs b/Engine/Plugins/Runtime/PacketHandlers/AESGCMHandlerComponent/Source/AESGCMHandlerComponent.Build.cs
--- a/Engine/Plugins/Runtime/PacketHandlers/AESGCMHandlerComponent/Source/AESGCMHandlerComponent.Build.cs
+++ b/Engine/Plugins/Runtime/PacketHandlers/AESGCMHandlerComponent/Source/AESGCMHandlerComponent.Build.cs
@@ -9,11 +9,6 @@ public class AESGCMHandlerComponent : ModuleRules
 
 	public AESGCMHandlerComponent(ReadOnlyTargetRules Target) : base(Target)
 	{
-		PrivateIncludePaths.AddRange(
-			new string[] {
-			}
-			);
-
 		PublicDependencyModuleNames.AddRange(
 			new string[]
 			{
diff --git a/Engine/Plugins/Runtime/PacketHandlers/AESHandlerComponent/Source/AESHandlerComponent.Build.cs b/Engine/Plugins/Runtime/PacketHandlers/AESHandlerComponent/Source/AESHandlerComponent.Build.cs
--- a/Engine/Plugins/Runtime/PacketHandlers/AESHandlerComponent/Source/AESHandlerComponent.Build.cs
+++ b/Engine/Plugins/Runtime/PacketHandlers/AESHandlerComponent/Source/AESHandlerComponent.Build.cs
@@ -9,11 +9,6 @@ public class AESHandlerComponent : ModuleRules
 
 	public AESHandlerComponent(ReadOnlyTargetRules Target) : base(Target)
 	{
-		PrivateIncludePaths.AddRange(
-			new string[] {
-			}
-			);
-
 		PublicDependencyModuleNames.AddRange(
 			new string[]
 			{
diff --git a/Engine/Plugins/Runtime/RenderTrace/Source/RenderTrace.Build.cs b/Engine/Plugins/Runtime/RenderTrace/Source/RenderTrace.Build.cs
--- a/Engine/Plugins/Runtime/RenderTrace/Source/RenderTrace.Build.cs
+++ b/Engine/Plugins/Runtime/RenderTrace/Source/RenderTrace.Build.cs
@@ -7,13 +7,6 @@ public class RenderTrace : ModuleRules
 {
 	public RenderTrace(ReadOnlyTargetRules Target) : base(Target)
 	{
-		PrivateIncludePaths.AddRange(
-			new string[] {
-				"Private",
-                //"../Shaders/Shared"
-            }
-		);
-
 		PrivateIncludePathModuleNames.AddRange(
 			new string[] {
 				"Renderer",
diff --git a/Engine/Plugins/Runtime/Soundscape/Source/SoundScape/SoundScape.Build.cs b/Engine/Plugins/Runtime/Soundscape/Source/SoundScape/SoundScape.Build.cs
--- a/Engine/Plugins/Runtime/Soundscape/Source/SoundScape/SoundScape.Build.cs
+++ b/Engine/Plugins/Runtime/Soundscape/Source/SoundScape/SoundScape.Build.cs
@@ -6,22 +6,6 @@ namespace UnrealBuildTool.Rules
 	{
         public Soundscape(ReadOnlyTargetRules Target) : base(Target)
         {
-			PCHUsage = ModuleRules.PCHUsageMode.UseExplicitOrSharedPCHs;
-
-			PublicIncludePaths.AddRange(
-				new string[] {
-					// ... add public include paths required here ...
-				}
-				);
-
-
-			PrivateIncludePaths.AddRange(
-				new string[] {
-					// ... add other private include paths required here ...
-				}
-				);
-
-
 			PublicDependencyModuleNames.AddRange(
 				new string[]
 				{
@@ -42,14 +26,6 @@ namespace UnrealBuildTool.Rules
 					// ... add private dependencies that you statically link with here ...	
 				}
 				);
-
-
-			DynamicallyLoadedModuleNames.AddRange(
-				new string[]
-				{
-					// ... add any modules that your module loads dynamically here ...
-				}
-				);
 		}
 	}
 }
\ No newline at end of file
diff --git a/Engine/Plugins/Runtime/Soundscape/Source/SoundScapeEditor/SoundScapeEditor.Build.cs b/Engine/Plugins/Runtime/Soundscape/Source/SoundScapeEditor/SoundScapeEditor.Build.cs
--- a/Engine/Plugins/Runtime/Soundscape/Source/SoundScapeEditor/SoundScapeEditor.Build.cs
+++ b/Engine/Plugins/Runtime/Soundscape/Source/SoundScapeEditor/SoundScapeEditor.Build.cs
@@ -6,23 +6,6 @@ namespace UnrealBuildTool.Rules
 	{
 		public SoundscapeEditor(ReadOnlyTargetRules Target) : base(Target)
         {
-
-			PCHUsage = ModuleRules.PCHUsageMode.UseExplicitOrSharedPCHs;
-
-			PublicIncludePaths.AddRange(
-				new string[] {
-					// ... add public include paths required here ...
-				}
-				);
-
-
-			PrivateIncludePaths.AddRange(
-				new string[] {
-					// ... add other private include paths required here ...
-				}
-				);
-
-
 			PublicDependencyModuleNames.AddRange(
 				new string[]
 				{
@@ -47,18 +30,6 @@ namespace UnrealBuildTool.Rules
 					// ... add private dependencies that you statically link with here ...	
 				}
 				);
-
-
-			DynamicallyLoadedModuleNames.AddRange(
-				new string[]
-				{
-					// ... add any modules that your module loads dynamically here ...
-				}
-				);
-
-			PrivateIncludePathModuleNames.AddRange(
-			new string[] {
-            });
 		}
 	}
 }
\ No newline at end of file
diff --git a/Engine/Plugins/Runtime/Steam/SteamAudio/Source/SteamAudio/SteamAudio.Build.cs b/Engine/Plugins/Runtime/Steam/SteamAudio/Source/SteamAudio/SteamAudio.Build.cs
--- a/Engine/Plugins/Runtime/Steam/SteamAudio/Source/SteamAudio/SteamAudio.Build.cs
+++ b/Engine/Plugins/Runtime/Steam/SteamAudio/Source/SteamAudio/SteamAudio.Build.cs
@@ -10,12 +10,6 @@ namespace UnrealBuildTool.Rules
 	{
 		public SteamAudio(ReadOnlyTargetRules Target) : base(Target)
 		{
-			PrivateIncludePaths.AddRange(
-				new string[] {
-					"SteamAudio/Private",
-				}
-			);
-
 			PublicDependencyModuleNames.AddRange(
 				new string[] {
 					"Core",
diff --git a/Engine/Plugins/Runtime/Steam/SteamAudio/Source/SteamAudioEditor/SteamAudioEditor.Build.cs b/Engine/Plugins/Runtime/Steam/SteamAudio/Source/SteamAudioEditor/SteamAudioEditor.Build.cs
--- a/Engine/Plugins/Runtime/Steam/SteamAudio/Source/SteamAudioEditor/SteamAudioEditor.Build.cs
+++ b/Engine/Plugins/Runtime/Steam/SteamAudio/Source/SteamAudioEditor/SteamAudioEditor.Build.cs
@@ -15,12 +15,6 @@ namespace UnrealBuildTool.Rules
 				}
 			);
 
-			PublicIncludePaths.AddRange(
-				new string[] {
-				}
-			);
-
-
 			PublicDependencyModuleNames.AddRange(
 				new string[] {
 					"Core",
diff --git a/Engine/Plugins/Runtime/ZoneGraph/Source/ZoneGraphDebug/ZoneGraphDebug.Build.cs b/Engine/Plugins/Runtime/ZoneGraph/Source/ZoneGraphDebug/ZoneGraphDebug.Build.cs
--- a/Engine/Plugins/Runtime/ZoneGraph/Source/ZoneGraphDebug/ZoneGraphDebug.Build.cs
+++ b/Engine/Plugins/Runtime/ZoneGraph/Source/ZoneGraphDebug/ZoneGraphDebug.Build.cs
@@ -15,11 +15,6 @@ namespace UnrealBuildTool.Rules
 			}
 			);
 
-			PrivateIncludePaths.AddRange(
-			new string[] {
-			}
-			);
-
 			PublicDependencyModuleNames.AddRange(
 			new string[] {
 				"Core",
diff --git a/Engine/Plugins/Tests/AutomationDriverTests/Source/AutomationDriverTests/AutomationDriverTests.Build.cs b/Engine/Plugins/Tests/AutomationDriverTests/Source/AutomationDriverTests/AutomationDriverTests.Build.cs
--- a/Engine/Plugins/Tests/AutomationDriverTests/Source/AutomationDriverTests/AutomationDriverTests.Build.cs
+++ b/Engine/Plugins/Tests/AutomationDriverTests/Source/AutomationDriverTests/AutomationDriverTests.Build.cs
@@ -6,21 +6,6 @@ public class AutomationDriverTests : ModuleRules
 {
 	public AutomationDriverTests(ReadOnlyTargetRules Target) : base(Target)
 	{
-		
-		PublicIncludePaths.AddRange(
-			new string[] {
-				// ... add public include paths required here ...
-			}
-			);
-				
-		
-		PrivateIncludePaths.AddRange(
-			new string[] {
-				// ... add other private include paths required here ...
-			}
-			);
-			
-		
 		PublicDependencyModuleNames.AddRange(
 			new string[]
 			{
@@ -49,13 +34,5 @@ public class AutomationDriverTests : ModuleRules
 				// ... add private dependencies that you statically link with here ...	
 			}
 			);
-		
-		
-		DynamicallyLoadedModuleNames.AddRange(
-			new string[]
-			{
-				// ... add any modules that your module loads dynamically here ...
-			}
-			);
 	}
 }
diff --git a/Engine/Plugins/Tests/EditorTests/Source/EditorTests/EditorTests.Build.cs b/Engine/Plugins/Tests/EditorTests/Source/EditorTests/EditorTests.Build.cs
--- a/Engine/Plugins/Tests/EditorTests/Source/EditorTests/EditorTests.Build.cs
+++ b/Engine/Plugins/Tests/EditorTests/Source/EditorTests/EditorTests.Build.cs
@@ -6,22 +6,6 @@ public class EditorTests : ModuleRules
 {
 	public EditorTests(ReadOnlyTargetRules Target) : base(Target)
 	{
-
-		PublicIncludePaths.AddRange(
-			new string[] {
-				// ... add public include paths required here ...
-			}
-		);
-
-
-		PrivateIncludePaths.AddRange(
-			new string[] {
-				"EditorTests/Private",
-				// ... add other private include paths required here ...
-			}
-		);
-
-
 		PublicDependencyModuleNames.AddRange(
 			new string[]
 			{
@@ -30,7 +14,6 @@ public class EditorTests : ModuleRules
 			}
 		);
 
-
 		PrivateDependencyModuleNames.AddRange(
 			new string[]
 			{
@@ -73,13 +56,5 @@ public class EditorTests : ModuleRules
 				//"SubobjectDataInterface",
 			}
 		);
-
-
-		DynamicallyLoadedModuleNames.AddRange(
-			new string[]
-			{
-				// ... add any modules that your module loads dynamically here ...
-			}
-		);
 	}
 }
diff --git a/Engine/Plugins/Tests/RHITests/Source/RHITests/RHITests.Build.cs b/Engine/Plugins/Tests/RHITests/Source/RHITests/RHITests.Build.cs
--- a/Engine/Plugins/Tests/RHITests/Source/RHITests/RHITests.Build.cs
+++ b/Engine/Plugins/Tests/RHITests/Source/RHITests/RHITests.Build.cs
@@ -5,14 +5,7 @@ using UnrealBuildTool;
 public class RHITests : ModuleRules
 {
 	public RHITests(ReadOnlyTargetRules Target) : base(Target)
-	{
-		
-		PublicIncludePaths.AddRange(
-			new string[] {
-				// ... add public include paths required here ...
-			}
-		);
-				
+	{				
 		PrivateIncludePaths.AddRange(
 			new string[] {
 				"../../../../Source/Runtime/Engine/",
@@ -42,13 +35,5 @@ public class RHITests : ModuleRules
 				// ... add private dependencies that you statically link with here ...	
 			}
 		);
-		
-		
-		DynamicallyLoadedModuleNames.AddRange(
-			new string[]
-			{
-				// ... add any modules that your module loads dynamically here ...
-			}
-		);
 	}
 }
diff --git a/Engine/Plugins/Tests/RuntimeTests/Source/RuntimeTests/RuntimeTests.Build.cs b/Engine/Plugins/Tests/RuntimeTests/Source/RuntimeTests/RuntimeTests.Build.cs
--- a/Engine/Plugins/Tests/RuntimeTests/Source/RuntimeTests/RuntimeTests.Build.cs
+++ b/Engine/Plugins/Tests/RuntimeTests/Source/RuntimeTests/RuntimeTests.Build.cs
@@ -6,21 +6,6 @@ public class RuntimeTests : ModuleRules
 {
 	public RuntimeTests(ReadOnlyTargetRules Target) : base(Target)
 	{
-		
-		PublicIncludePaths.AddRange(
-			new string[] {
-				// ... add public include paths required here ...
-			}
-			);
-				
-		
-		PrivateIncludePaths.AddRange(
-			new string[] {
-				// ... add other private include paths required here ...
-			}
-			);
-			
-		
 		PublicDependencyModuleNames.AddRange(
 			new string[]
 			{
@@ -43,13 +28,5 @@ public class RuntimeTests : ModuleRules
 				// ... add private dependencies that you statically link with here ...	
 			}
 			);
-		
-		
-		DynamicallyLoadedModuleNames.AddRange(
-			new string[]
-			{
-				// ... add any modules that your module loads dynamically here ...
-			}
-			);
 	}
 }
diff --git a/Engine/Plugins/VirtualProduction/LiveLinkCamera/Source/LiveLinkCamera/LiveLinkCamera.build.cs b/Engine/Plugins/VirtualProduction/LiveLinkCamera/Source/LiveLinkCamera/LiveLinkCamera.build.cs
--- a/Engine/Plugins/VirtualProduction/LiveLinkCamera/Source/LiveLinkCamera/LiveLinkCamera.build.cs
+++ b/Engine/Plugins/VirtualProduction/LiveLinkCamera/Source/LiveLinkCamera/LiveLinkCamera.build.cs
@@ -6,21 +6,6 @@ public class LiveLinkCamera : ModuleRules
 {
 	public LiveLinkCamera(ReadOnlyTargetRules Target) : base(Target)
 	{
-		PCHUsage = ModuleRules.PCHUsageMode.UseExplicitOrSharedPCHs;
-		
-		PublicIncludePaths.AddRange(
-			new string[] {
-			}
-			);
-				
-		
-		PrivateIncludePaths.AddRange(
-			new string[] 
-			{
-			}
-		);
-			
-		
 		PublicDependencyModuleNames.AddRange(
 			new string[]
 			{
@@ -34,12 +19,5 @@ public class LiveLinkCamera : ModuleRules
 				"LiveLinkComponents",
 			}
 		);
-			
-		
-		PrivateDependencyModuleNames.AddRange(
-			new string[]
-			{
-			}
-		);
 	}
 }
diff --git a/Engine/Plugins/VirtualProduction/LiveLinkCamera/Source/LiveLinkCameraEditor/LiveLinkCameraEditor.build.cs b/Engine/Plugins/VirtualProduction/LiveLinkCamera/Source/LiveLinkCameraEditor/LiveLinkCameraEditor.build.cs
--- a/Engine/Plugins/VirtualProduction/LiveLinkCamera/Source/LiveLinkCameraEditor/LiveLinkCameraEditor.build.cs
+++ b/Engine/Plugins/VirtualProduction/LiveLinkCamera/Source/LiveLinkCameraEditor/LiveLinkCameraEditor.build.cs
@@ -5,28 +5,7 @@ using UnrealBuildTool;
 public class LiveLinkCameraEditor : ModuleRules
 {
 	public LiveLinkCameraEditor(ReadOnlyTargetRules Target) : base(Target)
-	{
-		PCHUsage = ModuleRules.PCHUsageMode.UseExplicitOrSharedPCHs;
-		
-		PublicIncludePaths.AddRange(
-			new string[] {
-			}
-			);
-				
-		
-		PrivateIncludePaths.AddRange(
-			new string[] {
-			}
-		);
-			
-		
-		PublicDependencyModuleNames.AddRange(
-			new string[]
-			{
-				
-			}
-		);
-		
+	{		
 		PrivateDependencyModuleNames.AddRange(
 			new string[]
 			{
diff --git a/Engine/Plugins/VirtualProduction/LiveLinkCamera/Source/LiveLinkCameraRecording/LiveLinkCameraRecording.build.cs b/Engine/Plugins/VirtualProduction/LiveLinkCamera/Source/LiveLinkCameraRecording/LiveLinkCameraRecording.build.cs
--- a/Engine/Plugins/VirtualProduction/LiveLinkCamera/Source/LiveLinkCameraRecording/LiveLinkCameraRecording.build.cs
+++ b/Engine/Plugins/VirtualProduction/LiveLinkCamera/Source/LiveLinkCameraRecording/LiveLinkCameraRecording.build.cs
@@ -5,27 +5,7 @@ using UnrealBuildTool;
 public class LiveLinkCameraRecording : ModuleRules
 {
 	public LiveLinkCameraRecording(ReadOnlyTargetRules Target) : base(Target)
-	{
-		PCHUsage = ModuleRules.PCHUsageMode.UseExplicitOrSharedPCHs;
-
-		PublicIncludePaths.AddRange(
-			new string[] {
-			}
-			);
-				
-		
-		PrivateIncludePaths.AddRange(
-			new string[] {
-			}
-		);
-			
-		
-		PublicDependencyModuleNames.AddRange(
-			new string[]
-			{
-			}
-		);
-		
+	{		
 		PrivateDependencyModuleNames.AddRange(
 			new string[]
 			{
diff --git a/Engine/Plugins/VirtualProduction/SequencerPlaylists/Source/SequencerPlaylists/SequencerPlaylists.Build.cs b/Engine/Plugins/VirtualProduction/SequencerPlaylists/Source/SequencerPlaylists/SequencerPlaylists.Build.cs
--- a/Engine/Plugins/VirtualProduction/SequencerPlaylists/Source/SequencerPlaylists/SequencerPlaylists.Build.cs
+++ b/Engine/Plugins/VirtualProduction/SequencerPlaylists/Source/SequencerPlaylists/SequencerPlaylists.Build.cs
@@ -6,14 +6,6 @@ public class SequencerPlaylists : ModuleRules
 {
 	public SequencerPlaylists(ReadOnlyTargetRules Target) : base(Target)
 	{
-		PCHUsage = ModuleRules.PCHUsageMode.UseExplicitOrSharedPCHs;
-
-		PublicIncludePaths.AddRange(new string[] {
-		});
-
-		PrivateIncludePaths.AddRange(new string[] {
-		});
-
 		PublicDependencyModuleNames.AddRange(new string[] {
 			"Core",
 		});
diff --git a/Engine/Plugins/VirtualProduction/Takes/Source/TakeMovieScene/TakeMovieScene.Build.cs b/Engine/Plugins/VirtualProduction/Takes/Source/TakeMovieScene/TakeMovieScene.Build.cs
--- a/Engine/Plugins/VirtualProduction/Takes/Source/TakeMovieScene/TakeMovieScene.Build.cs
+++ b/Engine/Plugins/VirtualProduction/Takes/Source/TakeMovieScene/TakeMovieScene.Build.cs
@@ -20,12 +20,5 @@ public class TakeMovieScene : ModuleRules
 				"Engine",
             }
         );
-
-		PrivateIncludePaths.AddRange(
-			new string[] {
-				"TakeMovieScene/Private",
-				"TakeMovieScene/Public",
-			}
-		);
 	}
 }
diff --git a/Engine/Plugins/VirtualProduction/Takes/Source/TakeRecorderSources/TakeRecorderSources.Build.cs b/Engine/Plugins/VirtualProduction/Takes/Source/TakeRecorderSources/TakeRecorderSources.Build.cs
--- a/Engine/Plugins/VirtualProduction/Takes/Source/TakeRecorderSources/TakeRecorderSources.Build.cs
+++ b/Engine/Plugins/VirtualProduction/Takes/Source/TakeRecorderSources/TakeRecorderSources.Build.cs
@@ -38,12 +38,5 @@ public class TakeRecorderSources : ModuleRules
                 "TakeTrackRecorders",
             }
         );
-
-		PrivateIncludePaths.AddRange(
-			new string[] {
-				"TakeRecorderSources/Private",
-				"TakeRecorderSources/Public",
-			}
-		);
 	}
 }
diff --git a/Engine/Plugins/VirtualProduction/Takes/Source/TakeSequencer/TakeSequencer.Build.cs b/Engine/Plugins/VirtualProduction/Takes/Source/TakeSequencer/TakeSequencer.Build.cs
--- a/Engine/Plugins/VirtualProduction/Takes/Source/TakeSequencer/TakeSequencer.Build.cs
+++ b/Engine/Plugins/VirtualProduction/Takes/Source/TakeSequencer/TakeSequencer.Build.cs
@@ -27,12 +27,5 @@ public class TakeSequencer : ModuleRules
 				"Engine",
             }
         );
-
-		PrivateIncludePaths.AddRange(
-			new string[] {
-				"TakeSequencer/Private",
-				"TakeSequencer/Public",
-			}
-		);
 	}
 }
diff --git a/Engine/Plugins/VirtualProduction/Takes/Source/TakeTrackRecorders/TakeTrackRecorders.Build.cs b/Engine/Plugins/VirtualProduction/Takes/Source/TakeTrackRecorders/TakeTrackRecorders.Build.cs
--- a/Engine/Plugins/VirtualProduction/Takes/Source/TakeTrackRecorders/TakeTrackRecorders.Build.cs
+++ b/Engine/Plugins/VirtualProduction/Takes/Source/TakeTrackRecorders/TakeTrackRecorders.Build.cs
@@ -20,12 +20,5 @@ public class TakeTrackRecorders : ModuleRules
 				"UnrealEd",
             }
         );
-
-		PrivateIncludePaths.AddRange(
-			new string[] {
-				"TakeTrackRecorders/Private",
-				"TakeTrackRecorders/Public",
-			}
-		);
 	}
 }
diff --git a/Engine/Plugins/VirtualProduction/Takes/Source/TakesCore/TakesCore.Build.cs b/Engine/Plugins/VirtualProduction/Takes/Source/TakesCore/TakesCore.Build.cs
--- a/Engine/Plugins/VirtualProduction/Takes/Source/TakesCore/TakesCore.Build.cs
+++ b/Engine/Plugins/VirtualProduction/Takes/Source/TakesCore/TakesCore.Build.cs
@@ -29,12 +29,5 @@ public class TakesCore : ModuleRules
 
             }
         );
-
-		PrivateIncludePaths.AddRange(
-			new string[] {
-				"TakesCore/Private",
-				"TakesCore/Public",
-			}
-		);
 	}
 }
diff --git a/Engine/Source/Developer/AITestSuite/AITestSuite.Build.cs b/Engine/Source/Developer/AITestSuite/AITestSuite.Build.cs
--- a/Engine/Source/Developer/AITestSuite/AITestSuite.Build.cs
+++ b/Engine/Source/Developer/AITestSuite/AITestSuite.Build.cs
@@ -6,19 +6,6 @@ namespace UnrealBuildTool.Rules
 	{
         public AITestSuite(ReadOnlyTargetRules Target) : base(Target)
 		{
-			PublicIncludePaths.AddRange(
-				new string[] {
-					// ... add public include paths required here ...
-				}
-				);
-
-			PrivateIncludePaths.AddRange(
-				new string[] {
-					"Developer/AITestSuite/Private",
-					// ... add other private include paths required here ...
-				}
-				);
-
 			PublicDependencyModuleNames.AddRange(
 				new string[]
 				{
@@ -30,20 +17,6 @@ namespace UnrealBuildTool.Rules
                 }
                 );
 
-			PrivateDependencyModuleNames.AddRange(
-				new string[]
-				{
-					// ... add private dependencies that you statically link with here ...
-				}
-				);
-
-			DynamicallyLoadedModuleNames.AddRange(
-				new string[]
-				{
-					// ... add any modules that your module loads dynamically here ...
-				}
-				);
-
 			PrecompileForTargets = PrecompileTargetsType.Any;
 		}
 	}
diff --git a/Engine/Source/Developer/Android/AndroidTargetPlatform/AndroidTargetPlatform.Build.cs b/Engine/Source/Developer/Android/AndroidTargetPlatform/AndroidTargetPlatform.Build.cs
--- a/Engine/Source/Developer/Android/AndroidTargetPlatform/AndroidTargetPlatform.Build.cs
+++ b/Engine/Source/Developer/Android/AndroidTargetPlatform/AndroidTargetPlatform.Build.cs
@@ -35,11 +35,5 @@ public class AndroidTargetPlatform : ModuleRules
         }
 
         PublicDefinitions.Add("WITH_OGGVORBIS=1");
-
-		PrivateIncludePaths.AddRange(
-			new string[]
-			{
-            }
-        );
 	}
 }
diff --git a/Engine/Source/Developer/AnimationWidgets/AnimationWidgets.Build.cs b/Engine/Source/Developer/AnimationWidgets/AnimationWidgets.Build.cs
--- a/Engine/Source/Developer/AnimationWidgets/AnimationWidgets.Build.cs
+++ b/Engine/Source/Developer/AnimationWidgets/AnimationWidgets.Build.cs
@@ -19,16 +19,5 @@ public class AnimationWidgets : ModuleRules
 				"AnimationCore",
 			}
 		);
-
-		PrivateDependencyModuleNames.AddRange(
-			new string[] {
-			}
-		);
-
-		PrivateIncludePaths.AddRange(
-			new string[] {
-				"Developer/AnimationWidgets/Private",
-			}
-		);
 	}
 }
diff --git a/Engine/Source/Developer/AssetTools/AssetTools.Build.cs b/Engine/Source/Developer/AssetTools/AssetTools.Build.cs
--- a/Engine/Source/Developer/AssetTools/AssetTools.Build.cs
+++ b/Engine/Source/Developer/AssetTools/AssetTools.Build.cs
@@ -6,8 +6,6 @@ public class AssetTools : ModuleRules
 {
 	public AssetTools(ReadOnlyTargetRules Target) : base(Target)
 	{
-		PrivateIncludePaths.Add("Developer/AssetTools/Private");
-
 		PublicDependencyModuleNames.AddRange(
 			new string[] {
 				"Core",
diff --git a/Engine/Source/Developer/BlankModule/BlankModule.Build.cs b/Engine/Source/Developer/BlankModule/BlankModule.Build.cs
--- a/Engine/Source/Developer/BlankModule/BlankModule.Build.cs
+++ b/Engine/Source/Developer/BlankModule/BlankModule.Build.cs
@@ -14,7 +14,6 @@ namespace UnrealBuildTool.Rules
 
 			PrivateIncludePaths.AddRange(
 				new string[] {
-					"Developer/BlankModule/Private",
 					// ... add other private include paths required here ...
 				}
 				);
diff --git a/Engine/Source/Developer/DesktopPlatform/DesktopPlatform.Build.cs b/Engine/Source/Developer/DesktopPlatform/DesktopPlatform.Build.cs
--- a/Engine/Source/Developer/DesktopPlatform/DesktopPlatform.Build.cs
+++ b/Engine/Source/Developer/DesktopPlatform/DesktopPlatform.Build.cs
@@ -6,8 +6,6 @@ public class DesktopPlatform : ModuleRules
 {
 	public DesktopPlatform(ReadOnlyTargetRules Target) : base(Target)
 	{
-		PrivateIncludePaths.Add("Developer/DesktopPlatform/Private");
-
 		PrivateIncludePathModuleNames.AddRange(
 			new string[] {
 				"DerivedDataCache",
diff --git a/Engine/Source/Developer/DirectoryWatcher/DirectoryWatcher.Build.cs b/Engine/Source/Developer/DirectoryWatcher/DirectoryWatcher.Build.cs
--- a/Engine/Source/Developer/DirectoryWatcher/DirectoryWatcher.Build.cs
+++ b/Engine/Source/Developer/DirectoryWatcher/DirectoryWatcher.Build.cs
@@ -6,8 +6,6 @@ public class DirectoryWatcher : ModuleRules
 {
 	public DirectoryWatcher(ReadOnlyTargetRules Target) : base(Target)
 	{
-		PrivateIncludePaths.Add("Developer/DirectoryWatcher/Private");
-
 		PrivateDependencyModuleNames.Add("Core");
 	}
 }
diff --git a/Engine/Source/Developer/DrawPrimitiveDebugger/DrawPrimitiveDebugger.Build.cs b/Engine/Source/Developer/DrawPrimitiveDebugger/DrawPrimitiveDebugger.Build.cs
--- a/Engine/Source/Developer/DrawPrimitiveDebugger/DrawPrimitiveDebugger.Build.cs
+++ b/Engine/Source/Developer/DrawPrimitiveDebugger/DrawPrimitiveDebugger.Build.cs
@@ -1,4 +1,4 @@
-﻿// Copyright Epic Games, Inc. All Rights Reserved.
+// Copyright Epic Games, Inc. All Rights Reserved.
 
 namespace UnrealBuildTool.Rules
 {
@@ -25,7 +25,6 @@ namespace UnrealBuildTool.Rules
 
 			PrivateIncludePaths.AddRange(
 				new string[] {
-					"Developer/DrawPrimitiveDebugger/Private",
 					"Developer/Settings/Public",
 				});
 
diff --git a/Engine/Source/Developer/FunctionalTesting/FunctionalTesting.Build.cs b/Engine/Source/Developer/FunctionalTesting/FunctionalTesting.Build.cs
--- a/Engine/Source/Developer/FunctionalTesting/FunctionalTesting.Build.cs
+++ b/Engine/Source/Developer/FunctionalTesting/FunctionalTesting.Build.cs
@@ -45,13 +45,6 @@ public class FunctionalTesting : ModuleRules
 			);
 		}
 
-        PrivateIncludePaths.AddRange(
-            new string[]
-            {
-                "Developer/FunctionalTesting/Private",
-			}
-		);
-
 		PrivateIncludePathModuleNames.AddRange(
 			new string[] {
 				"ImageWrapper",
diff --git a/Engine/Source/Developer/GameplayDebugger/GameplayDebugger.Build.cs b/Engine/Source/Developer/GameplayDebugger/GameplayDebugger.Build.cs
--- a/Engine/Source/Developer/GameplayDebugger/GameplayDebugger.Build.cs
+++ b/Engine/Source/Developer/GameplayDebugger/GameplayDebugger.Build.cs
@@ -29,7 +29,6 @@ namespace UnrealBuildTool.Rules
 
 			PrivateIncludePaths.AddRange(
 				new string[] {
-					"Developer/GameplayDebugger/Private",
 					"Developer/Settings/Public",
 				});
 
diff --git a/Engine/Source/Developer/GraphColor/GraphColor.Build.cs b/Engine/Source/Developer/GraphColor/GraphColor.Build.cs
--- a/Engine/Source/Developer/GraphColor/GraphColor.Build.cs
+++ b/Engine/Source/Developer/GraphColor/GraphColor.Build.cs
@@ -6,9 +6,6 @@ public class GraphColor : ModuleRules
 {
 	public GraphColor(ReadOnlyTargetRules Target) : base(Target)
 	{
-        PublicIncludePaths.Add("Developer/GraphColor/Public");
-		//PrivateIncludePaths.Add("Developer/GraphColor/Private");
-
         PrivateDependencyModuleNames.AddRange(new string[] { "Core" });
 	}
 }
diff --git a/Engine/Source/Developer/IOS/IOSTargetPlatform/IOSTargetPlatform.Build.cs b/Engine/Source/Developer/IOS/IOSTargetPlatform/IOSTargetPlatform.Build.cs
--- a/Engine/Source/Developer/IOS/IOSTargetPlatform/IOSTargetPlatform.Build.cs
+++ b/Engine/Source/Developer/IOS/IOSTargetPlatform/IOSTargetPlatform.Build.cs
@@ -30,8 +30,6 @@ public class IOSTargetPlatform : ModuleRules
 			}
 		);
 
-		PrivateIncludePaths.Add("Developer/IOS/IOSTargetPlatform/Private");
-
 		//This is somehow necessary for getting iOS to build on, at least, windows. It seems like the target platform is included for cooking, and thus it requirtes a bunch of other info.
 		PublicIncludePaths.AddRange(
 			new string[]
diff --git a/Engine/Source/Developer/Linux/LinuxTargetPlatform/LinuxTargetPlatform.Build.cs b/Engine/Source/Developer/Linux/LinuxTargetPlatform/LinuxTargetPlatform.Build.cs
--- a/Engine/Source/Developer/Linux/LinuxTargetPlatform/LinuxTargetPlatform.Build.cs
+++ b/Engine/Source/Developer/Linux/LinuxTargetPlatform/LinuxTargetPlatform.Build.cs
@@ -24,11 +24,6 @@ public class LinuxTargetPlatform : ModuleRules
 			}
 		);
 
-        PrivateIncludePaths.AddRange(
-            new string[] {
-			}
-        );
-
 		if (Target.bCompileAgainstEngine)
 		{
 			PrivateDependencyModuleNames.Add("Engine");
diff --git a/Engine/Source/Developer/LocalizationService/LocalizationService.Build.cs b/Engine/Source/Developer/LocalizationService/LocalizationService.Build.cs
--- a/Engine/Source/Developer/LocalizationService/LocalizationService.Build.cs
+++ b/Engine/Source/Developer/LocalizationService/LocalizationService.Build.cs
@@ -6,8 +6,6 @@ public class LocalizationService : ModuleRules
 {
 	public LocalizationService(ReadOnlyTargetRules Target) : base(Target)
 	{
-        PrivateIncludePaths.Add("Developer/LocalizationService/Private");
-
 		PublicDependencyModuleNames.AddRange(
 			new string[] {
 				"Core",
diff --git a/Engine/Source/Developer/Mac/MacTargetPlatform/MacTargetPlatform.Build.cs b/Engine/Source/Developer/Mac/MacTargetPlatform/MacTargetPlatform.Build.cs
--- a/Engine/Source/Developer/Mac/MacTargetPlatform/MacTargetPlatform.Build.cs
+++ b/Engine/Source/Developer/Mac/MacTargetPlatform/MacTargetPlatform.Build.cs
@@ -21,11 +21,6 @@ public class MacTargetPlatform : ModuleRules
 			}
 		);
 
-		PrivateIncludePaths.AddRange(
-			new string[] {
-			}
-		);
-
 		if (Target.bCompileAgainstEngine)
 		{
 			PrivateDependencyModuleNames.Add("Engine");
diff --git a/Engine/Source/Developer/Merge/Merge.Build.cs b/Engine/Source/Developer/Merge/Merge.Build.cs
--- a/Engine/Source/Developer/Merge/Merge.Build.cs
+++ b/Engine/Source/Developer/Merge/Merge.Build.cs
@@ -6,12 +6,6 @@ namespace UnrealBuildTool.Rules
 	{
         public Merge(ReadOnlyTargetRules Target) : base(Target)
 		{
-			PrivateIncludePaths.AddRange(
-				new string[] {
-					"Developer/Merge/Private",
-				}
-			);
-
 			PrivateDependencyModuleNames.AddRange(
 				new string[]
 				{
diff --git a/Engine/Source/Developer/MeshDescriptionOperations/MeshDescriptionOperations.Build.cs b/Engine/Source/Developer/MeshDescriptionOperations/MeshDescriptionOperations.Build.cs
--- a/Engine/Source/Developer/MeshDescriptionOperations/MeshDescriptionOperations.Build.cs
+++ b/Engine/Source/Developer/MeshDescriptionOperations/MeshDescriptionOperations.Build.cs
@@ -6,9 +6,6 @@ namespace UnrealBuildTool.Rules
     {
         public MeshDescriptionOperations(ReadOnlyTargetRules Target) : base(Target)
         {
-            PrivateIncludePaths.Add("Developer/MeshDescriptionOperations/Private");
-            PublicIncludePaths.Add("Developer/MeshDescriptionOperations/Public");
-
             PrivateDependencyModuleNames.AddRange(
                 new string[]
                 {
diff --git a/Engine/Source/Developer/ProfilerMessages/ProfilerMessages.Build.cs b/Engine/Source/Developer/ProfilerMessages/ProfilerMessages.Build.cs
--- a/Engine/Source/Developer/ProfilerMessages/ProfilerMessages.Build.cs
+++ b/Engine/Source/Developer/ProfilerMessages/ProfilerMessages.Build.cs
@@ -18,11 +18,6 @@ namespace UnrealBuildTool.Rules
 					"Core",
 				});
 
-			PrivateIncludePaths.AddRange(
-				new string[]
-				{
-				});
-
 			if (Target.Configuration != UnrealTargetConfiguration.Shipping)
 			{
 				PrecompileForTargets = PrecompileTargetsType.Any;
diff --git a/Engine/Source/Developer/ScreenShotComparisonTools/ScreenShotComparisonTools.Build.cs b/Engine/Source/Developer/ScreenShotComparisonTools/ScreenShotComparisonTools.Build.cs
--- a/Engine/Source/Developer/ScreenShotComparisonTools/ScreenShotComparisonTools.Build.cs
+++ b/Engine/Source/Developer/ScreenShotComparisonTools/ScreenShotComparisonTools.Build.cs
@@ -32,12 +32,6 @@ public class ScreenShotComparisonTools : ModuleRules
 			}
 		);
 
-		PrivateIncludePaths.AddRange(
-			new string[] {
-				"Developer/ScreenShotComparisonTools/Private"
-			}
-		);
-
 		if (Target.bCompileAgainstEngine && Target.Configuration != UnrealTargetConfiguration.Shipping)
 		{
 			PrecompileForTargets = PrecompileTargetsType.Any;
diff --git a/Engine/Source/Developer/Settings/Settings.Build.cs b/Engine/Source/Developer/Settings/Settings.Build.cs
--- a/Engine/Source/Developer/Settings/Settings.Build.cs
+++ b/Engine/Source/Developer/Settings/Settings.Build.cs
@@ -18,12 +18,6 @@ public class Settings : ModuleRules
 				"CoreUObject",
 			});
 
-		PrivateIncludePaths.AddRange(
-			new string[]
-			{
-				"Developer/Settings/Private",
-			});
-
 		PrecompileForTargets = PrecompileTargetsType.Any;
 	}
 }
diff --git a/Engine/Source/Developer/SlackIntegrations/SlackIntegrations.Build.cs b/Engine/Source/Developer/SlackIntegrations/SlackIntegrations.Build.cs
--- a/Engine/Source/Developer/SlackIntegrations/SlackIntegrations.Build.cs
+++ b/Engine/Source/Developer/SlackIntegrations/SlackIntegrations.Build.cs
@@ -6,10 +6,6 @@ public class SlackIntegrations : ModuleRules
 {
     public SlackIntegrations(ReadOnlyTargetRules Target) : base(Target)
     {
-		PublicIncludePaths.Add("Developer/SlackIntegrations/Public");
-
-		PrivateIncludePaths.Add("Developer/SlackIntegrations/Private");
-
 		PublicDependencyModuleNames.AddRange(
 			new string[] {
 				"Core",
diff --git a/Engine/Source/Developer/SlateFileDialogs/SlateFileDialogs.Build.cs b/Engine/Source/Developer/SlateFileDialogs/SlateFileDialogs.Build.cs
--- a/Engine/Source/Developer/SlateFileDialogs/SlateFileDialogs.Build.cs
+++ b/Engine/Source/Developer/SlateFileDialogs/SlateFileDialogs.Build.cs
@@ -16,12 +16,6 @@ public class SlateFileDialogs : ModuleRules
                 "DirectoryWatcher",
             }
         );
-    
-        PrivateIncludePaths.AddRange(
-            new string[] {
-                "Developer/SlateFileDialogs/Private",  
-            }
-        );
 
         PrivateIncludePathModuleNames.Add("TargetPlatform");
     }
diff --git a/Engine/Source/Developer/SourceCodeAccess/SourceCodeAccess.Build.cs b/Engine/Source/Developer/SourceCodeAccess/SourceCodeAccess.Build.cs
--- a/Engine/Source/Developer/SourceCodeAccess/SourceCodeAccess.Build.cs
+++ b/Engine/Source/Developer/SourceCodeAccess/SourceCodeAccess.Build.cs
@@ -19,7 +19,5 @@ public class SourceCodeAccess : ModuleRules
 				"Slate",
 			}
 		);
-
-		PrivateIncludePaths.Add("Developer/SourceCodeAccess/Private");
 	}
 }
diff --git a/Engine/Source/Developer/SourceControl/SourceControl.Build.cs b/Engine/Source/Developer/SourceControl/SourceControl.Build.cs
--- a/Engine/Source/Developer/SourceControl/SourceControl.Build.cs
+++ b/Engine/Source/Developer/SourceControl/SourceControl.Build.cs
@@ -6,8 +6,6 @@ public class SourceControl : ModuleRules
 {
 	public SourceControl(ReadOnlyTargetRules Target) : base(Target)
 	{
-        PrivateIncludePaths.Add("Developer/SourceControl/Private");
-
 		PublicDependencyModuleNames.AddRange(
 			new string[] {
 				"Core",
diff --git a/Engine/Source/Developer/StandaloneRenderer/StandaloneRenderer.Build.cs b/Engine/Source/Developer/StandaloneRenderer/StandaloneRenderer.Build.cs
--- a/Engine/Source/Developer/StandaloneRenderer/StandaloneRenderer.Build.cs
+++ b/Engine/Source/Developer/StandaloneRenderer/StandaloneRenderer.Build.cs
@@ -6,8 +6,6 @@ public class StandaloneRenderer : ModuleRules
 {
 	public StandaloneRenderer(ReadOnlyTargetRules Target) : base(Target)
 	{
-		PrivateIncludePaths.Add("Developer/StandaloneRenderer/Private");
-
 		PrivateDependencyModuleNames.AddRange(
 			new string[] {
 				"Core",
diff --git a/Engine/Source/Developer/ToolWidgets/ToolWidgets.Build.cs b/Engine/Source/Developer/ToolWidgets/ToolWidgets.Build.cs
--- a/Engine/Source/Developer/ToolWidgets/ToolWidgets.Build.cs
+++ b/Engine/Source/Developer/ToolWidgets/ToolWidgets.Build.cs
@@ -27,11 +27,5 @@ public class ToolWidgets : ModuleRules
 				"ApplicationCore"
 			}
 		);
-
-		PrivateIncludePaths.AddRange(
-			new string[] {
-				"Developer/ToolWidgets/Private",
-			}
-		);
 	}
 }
diff --git a/Engine/Source/Developer/TraceInsights/TraceInsights.Build.cs b/Engine/Source/Developer/TraceInsights/TraceInsights.Build.cs
--- a/Engine/Source/Developer/TraceInsights/TraceInsights.Build.cs
+++ b/Engine/Source/Developer/TraceInsights/TraceInsights.Build.cs
@@ -6,13 +6,6 @@ public class TraceInsights : ModuleRules
 {
 	public TraceInsights(ReadOnlyTargetRules Target) : base(Target)
 	{
-		PrivateIncludePaths.AddRange
-		(
-			new string[] {
-				"Developer/TraceInsights/Private",
-			}
-		);
-
 		PublicDependencyModuleNames.AddRange
 		(
 			new string[] {
diff --git a/Engine/Source/Editor/AnimGraph/AnimGraph.Build.cs b/Engine/Source/Editor/AnimGraph/AnimGraph.Build.cs
--- a/Engine/Source/Editor/AnimGraph/AnimGraph.Build.cs
+++ b/Engine/Source/Editor/AnimGraph/AnimGraph.Build.cs
@@ -6,13 +6,7 @@ using System.IO;
 public class AnimGraph : ModuleRules
 {
 	public AnimGraph(ReadOnlyTargetRules Target) : base(Target)
-	{
-        PrivateIncludePaths.AddRange(
-            new string[] {
-                "Editor/AnimGraph/Private",
-            }
-        );
-		
+	{		
 		OverridePackageType = PackageOverrideType.EngineDeveloper;
 
         PublicDependencyModuleNames.AddRange(
diff --git a/Engine/Source/Editor/AnimationBlueprintEditor/AnimationBlueprintEditor.Build.cs b/Engine/Source/Editor/AnimationBlueprintEditor/AnimationBlueprintEditor.Build.cs
--- a/Engine/Source/Editor/AnimationBlueprintEditor/AnimationBlueprintEditor.Build.cs
+++ b/Engine/Source/Editor/AnimationBlueprintEditor/AnimationBlueprintEditor.Build.cs
@@ -6,8 +6,6 @@ public class AnimationBlueprintEditor : ModuleRules
 {
 	public AnimationBlueprintEditor(ReadOnlyTargetRules Target) : base(Target)
 	{
-		PrivateIncludePaths.Add("Editor/AnimationBlueprintEditor/Private");	// For PCH includes (because they don't work with relative paths, yet)
-
 		PrivateIncludePathModuleNames.AddRange(
 			new string[] {
 				"AssetRegistry", 
diff --git a/Engine/Source/Editor/BlueprintEditorLibrary/BlueprintEditorLibrary.Build.cs b/Engine/Source/Editor/BlueprintEditorLibrary/BlueprintEditorLibrary.Build.cs
--- a/Engine/Source/Editor/BlueprintEditorLibrary/BlueprintEditorLibrary.Build.cs
+++ b/Engine/Source/Editor/BlueprintEditorLibrary/BlueprintEditorLibrary.Build.cs
@@ -6,18 +6,7 @@ using System.IO;
 public class BlueprintEditorLibrary : ModuleRules
 {
 	public BlueprintEditorLibrary(ReadOnlyTargetRules Target) : base(Target)
-	{
-		PCHUsage = ModuleRules.PCHUsageMode.UseExplicitOrSharedPCHs;
-		PublicIncludePaths.Add(Path.Combine(ModuleDirectory, "Public"));
-
-		PrivateIncludePaths.AddRange(
-			new string[] 
-			{
-				Path.Combine(ModuleDirectory, "Private"),
-			}
-		);
-			
-		
+	{		
 		PublicDependencyModuleNames.AddRange(
 			new string[]
 			{
diff --git a/Engine/Source/Editor/Blutility/Blutility.Build.cs b/Engine/Source/Editor/Blutility/Blutility.Build.cs
--- a/Engine/Source/Editor/Blutility/Blutility.Build.cs
+++ b/Engine/Source/Editor/Blutility/Blutility.Build.cs
@@ -6,8 +6,6 @@ public class Blutility : ModuleRules
 {
 	public Blutility(ReadOnlyTargetRules Target) : base(Target)
 	{
-		PrivateIncludePaths.Add("Editor/Blutility/Private");
-
 		PrivateIncludePathModuleNames.Add("AssetTools");
 
         PublicDependencyModuleNames.AddRange(new string[] {
diff --git a/Engine/Source/Editor/Cascade/Cascade.Build.cs b/Engine/Source/Editor/Cascade/Cascade.Build.cs
--- a/Engine/Source/Editor/Cascade/Cascade.Build.cs
+++ b/Engine/Source/Editor/Cascade/Cascade.Build.cs
@@ -6,12 +6,6 @@ public class Cascade : ModuleRules
 {
 	public Cascade(ReadOnlyTargetRules Target) : base(Target)
 	{
-		PrivateIncludePaths.AddRange(
-			new string[] {
-				"Editor/Cascade/Private"
-			}
-		);
-
 		PublicIncludePaths.AddRange(
 			new string[] {
 				"Editor/DistCurveEditor/Public",
diff --git a/Engine/Source/Editor/ComponentVisualizers/ComponentVisualizers.Build.cs b/Engine/Source/Editor/ComponentVisualizers/ComponentVisualizers.Build.cs
--- a/Engine/Source/Editor/ComponentVisualizers/ComponentVisualizers.Build.cs
+++ b/Engine/Source/Editor/ComponentVisualizers/ComponentVisualizers.Build.cs
@@ -6,8 +6,6 @@ public class ComponentVisualizers : ModuleRules
 {
 	public ComponentVisualizers(ReadOnlyTargetRules Target) : base(Target)
 	{
-		PrivateIncludePaths.Add("Editor/ComponentVisualizers/Private");	// For PCH includes (because they don't work with relative paths, yet)
-
 		PrivateDependencyModuleNames.AddRange(
 			new string[] {
 				"Core",
diff --git a/Engine/Source/Editor/DataLayerEditor/DataLayerEditor.Build.cs b/Engine/Source/Editor/DataLayerEditor/DataLayerEditor.Build.cs
--- a/Engine/Source/Editor/DataLayerEditor/DataLayerEditor.Build.cs
+++ b/Engine/Source/Editor/DataLayerEditor/DataLayerEditor.Build.cs
@@ -6,10 +6,6 @@ public class DataLayerEditor : ModuleRules
 {
 	public DataLayerEditor(ReadOnlyTargetRules Target) : base(Target)
 	{
-		PublicIncludePaths.Add("Editor/DataLayerEditor/Public");
-
-		PrivateIncludePaths.Add("Editor/DataLayerEditor/Private"); // For PCH includes (because they don't work with relative paths, yet)
-
 		PrivateDependencyModuleNames.AddRange(
 			new string[] {
 				"Core",
diff --git a/Engine/Source/Editor/DetailCustomizations/DetailCustomizations.Build.cs b/Engine/Source/Editor/DetailCustomizations/DetailCustomizations.Build.cs
--- a/Engine/Source/Editor/DetailCustomizations/DetailCustomizations.Build.cs
+++ b/Engine/Source/Editor/DetailCustomizations/DetailCustomizations.Build.cs
@@ -6,8 +6,6 @@ public class DetailCustomizations : ModuleRules
 {
 	public DetailCustomizations(ReadOnlyTargetRules Target) : base(Target)
 	{
-		PrivateIncludePaths.Add("Editor/DetailCustomizations/Private");	// For PCH includes (because they don't work with relative paths, yet)
-
 		PrivateDependencyModuleNames.AddRange(
 			new string[] {
 				"AppFramework",
diff --git a/Engine/Source/Editor/DeviceProfileServices/DeviceProfileServices.Build.cs b/Engine/Source/Editor/DeviceProfileServices/DeviceProfileServices.Build.cs
--- a/Engine/Source/Editor/DeviceProfileServices/DeviceProfileServices.Build.cs
+++ b/Engine/Source/Editor/DeviceProfileServices/DeviceProfileServices.Build.cs
@@ -7,8 +7,6 @@ public class DeviceProfileServices : ModuleRules
 
 	public DeviceProfileServices(ReadOnlyTargetRules Target) : base(Target)
 	{
-		PrivateIncludePaths.Add("Editor/DeviceProfileServices/Private");
-
 		PublicDependencyModuleNames.AddRange(
 			new string[] {
 				"Core",
diff --git a/Engine/Source/Editor/Documentation/Documentation.Build.cs b/Engine/Source/Editor/Documentation/Documentation.Build.cs
--- a/Engine/Source/Editor/Documentation/Documentation.Build.cs
+++ b/Engine/Source/Editor/Documentation/Documentation.Build.cs
@@ -6,18 +6,6 @@ namespace UnrealBuildTool.Rules
 	{
 		public Documentation(ReadOnlyTargetRules Target) : base(Target)
 		{
-			PublicIncludePaths.AddRange(
-				new string[] {
-					// ... add public include paths required here ...
-				}
-			);
-
-			PrivateIncludePaths.AddRange(
-				new string[] {
-					// ... add other private include paths required here ...
-				}
-			);
-
 			PublicDependencyModuleNames.AddRange(
 				new string[]
 				{
diff --git a/Engine/Source/Editor/EditorFramework/EditorFramework.Build.cs b/Engine/Source/Editor/EditorFramework/EditorFramework.Build.cs
--- a/Engine/Source/Editor/EditorFramework/EditorFramework.Build.cs
+++ b/Engine/Source/Editor/EditorFramework/EditorFramework.Build.cs
@@ -6,16 +6,6 @@ public class EditorFramework : ModuleRules
 {
 	public EditorFramework(ReadOnlyTargetRules Target) : base(Target)
 	{
-		PrivateIncludePaths.AddRange(
-			new string[] {
-			}
-		);
-
-		PrivateIncludePathModuleNames.AddRange(
-			new string[] {
-			}
-		);
-
 		PrivateDependencyModuleNames.AddRange(
 			new string[] {
 				"Core",
@@ -37,15 +27,5 @@ public class EditorFramework : ModuleRules
 				"TypedElementRuntime",
 			}
 		);
-
-		PublicIncludePathModuleNames.AddRange(
-			new string[] {
-			}
-		);
-
-		DynamicallyLoadedModuleNames.AddRange(
-            new string[] { 
-            }
-		);
 	}
 }
