Fix bug where Soundscape Palettes with empty/no playback conditions do not activate

robomerge[FN_EngineMerge] -UE5-RES

diff --git a/Engine/Plugins/Runtime/Soundscape/Source/SoundScape/Private/SoundscapeSubsystem.cpp b/Engine/Plugins/Runtime/Soundscape/Source/SoundScape/Private/SoundscapeSubsystem.cpp
--- a/Engine/Plugins/Runtime/Soundscape/Source/SoundScape/Private/SoundscapeSubsystem.cpp
+++ b/Engine/Plugins/Runtime/Soundscape/Source/SoundScape/Private/SoundscapeSubsystem.cpp
@@ -924,7 +924,7 @@ void USoundscapeSubsystem::UpdateState()
 
 		if (SoundscapePalette && ActiveSoundscapePalette)
 		{
-			if (SoundscapePalette->SoundscapePalettePlaybackConditions.Matches(SubsystemState))
+			if (SoundscapePalette->SoundscapePalettePlaybackConditions.IsEmpty() || SoundscapePalette->SoundscapePalettePlaybackConditions.Matches(SubsystemState))
 			{
 				ActivePalettesToKeep.Add(SoundscapePalette, ActiveSoundscapePalette);
 			}
@@ -947,7 +947,7 @@ void USoundscapeSubsystem::UpdateState()
 		if (USoundscapePalette* SoundscapePalette = *PaletteIterator)
 		{
 			// Evaluate if the Soundscape Palette matches
-			if (SoundscapePalette->SoundscapePalettePlaybackConditions.Matches(SubsystemState) && ActivePalettes.Contains(SoundscapePalette) == false)
+			if ((SoundscapePalette->SoundscapePalettePlaybackConditions.IsEmpty() || SoundscapePalette->SoundscapePalettePlaybackConditions.Matches(SubsystemState)) && ActivePalettes.Contains(SoundscapePalette) == false)
 			{
 				UWorld* World = GetWorld();
 
