Clarifying tooltip for SpawnSoundscapeColor

diff --git a/Engine/Plugins/Runtime/Soundscape/Source/SoundScape/Public/SoundScape.h b/Engine/Plugins/Runtime/Soundscape/Source/SoundScape/Public/SoundScape.h
--- a/Engine/Plugins/Runtime/Soundscape/Source/SoundScape/Public/SoundScape.h
+++ b/Engine/Plugins/Runtime/Soundscape/Source/SoundScape/Public/SoundScape.h
@@ -14,7 +14,7 @@ class SOUNDSCAPE_API USoundscapeBPFunctionLibrary : public UBlueprintFunctionLib
 {
 	GENERATED_BODY()
 
-	/** Spawns a new Soundscape Elemental Agent. Returns true if spawn was successful. */
+	/** Spawns, but does not play, a new Soundscape Elemental Agent. Returns true if spawn was successful. "Play" can be called on the resulting Active Soundscape Color*/
 	UFUNCTION(BlueprintCallable, Category = "Soundscape", meta = (WorldContext = "WorldContextObject"))
 	static bool SpawnSoundscapeColor(UObject* WorldContextObject, class USoundscapeColor* SoundscapeColorIn, UActiveSoundscapeColor*& ActiveSoundscapeColor);
 
