[UE][Audio][Soundscape] Fix tooltip errata on Soundscape Palette

diff --git a/Engine/Plugins/Runtime/Soundscape/Source/SoundScape/Public/SoundScapePalette.h b/Engine/Plugins/Runtime/Soundscape/Source/SoundScape/Public/SoundScapePalette.h
--- a/Engine/Plugins/Runtime/Soundscape/Source/SoundScape/Public/SoundScapePalette.h
+++ b/Engine/Plugins/Runtime/Soundscape/Source/SoundScape/Public/SoundScapePalette.h
@@ -29,11 +29,11 @@ struct SOUNDSCAPE_API FSoundscapePaletteColor
 	UPROPERTY(EditAnywhere, BlueprintReadWrite, Category = "Soundscape|Palette", meta = (ClampMin = "0.2", ClampMax = "4.0", UIMin = "0.2", UIMax = "4.0", SliderExponent = "3.0"))
 	float ColorPitch = 1.0f;
 
-	// Base Volume Scalar
+	// Volume Fade In Time
 	UPROPERTY(EditAnywhere, BlueprintReadWrite, Category = "Soundscape|Palette", meta = (ClampMin = "0.0", UIMin = "0.0"))
 	float ColorFadeIn = 1.0f;
 
-	// Base Volume Scalar
+	// Volume Fade Out Time
 	UPROPERTY(EditAnywhere, BlueprintReadWrite, Category = "Soundscape|Palette", meta = (ClampMin = "0.0", UIMin = "0.0"))
 	float ColorFadeOut = 1.0f;
 };
