[Soundscape] Adding Asynchronous Loading to Soundscape Palettes

[RN] Asynchronous loading added to Soundscape Palettes

diff --git a/Engine/Plugins/Runtime/Soundscape/Source/SoundScape/Private/SoundscapeSubsystem.cpp b/Engine/Plugins/Runtime/Soundscape/Source/SoundScape/Private/SoundscapeSubsystem.cpp
--- a/Engine/Plugins/Runtime/Soundscape/Source/SoundScape/Private/SoundscapeSubsystem.cpp
+++ b/Engine/Plugins/Runtime/Soundscape/Source/SoundScape/Private/SoundscapeSubsystem.cpp
@@ -10,6 +10,8 @@
 #include "SoundScapePalette.h"
 #include "AudioDevice.h"
 #include "TimerManager.h"
+#include "Engine/AssetManager.h"
+#include "Engine/StreamableManager.h"
 
 static FAutoConsoleCommandWithWorld GResetSoundscape(
 	TEXT("soundscape.ResetSoundscape"),
@@ -487,6 +489,17 @@ void USoundscapeSubsystem::Deinitialize()
 	UE_LOG(LogSoundscapeSubsystem, Verbose, TEXT("Subsystem Deinitializing"));
 
 	ActiveColorPointUpdateTimer.Invalidate();
+
+	for(const TPair<FName, TSharedPtr<FStreamableHandle>>& HandleSharedPtr : ActiveAssetLoadHandles)
+	{
+		if(HandleSharedPtr.Value.IsValid())
+		{
+			if (FStreamableHandle* HandlePtr = HandleSharedPtr.Value.Get())
+			{
+				HandlePtr->CancelHandle();
+			}
+		}
+	}
 }
 
 bool USoundscapeSubsystem::ShouldCreateSubsystem(UObject* Outer) const
@@ -593,6 +606,20 @@ bool USoundscapeSubsystem::RemovePaletteCollection(FName PaletteCollectionName)
 	{
 		if (UnloadPaletteCollection(PaletteCollectionName))
 		{
+			// Release the handle pointer if Palette Collection still loading
+			if (TSharedPtr<FStreamableHandle>* AssetHandlePtr = ActiveAssetLoadHandles.Find(PaletteCollectionName))
+			{
+				TSharedPtr<FStreamableHandle> AssetHandle = *AssetHandlePtr;
+
+				if (AssetHandle.IsValid())
+				{
+					AssetHandle.Get()->CancelHandle();
+					AssetHandle.Reset();
+				}
+
+				ActiveAssetLoadHandles.Remove(PaletteCollectionName);
+			}
+
 			// Unloaded, now remove
 			UnloadedPaletteCollections.Remove(PaletteCollectionName);
 
@@ -613,30 +640,39 @@ bool USoundscapeSubsystem::LoadPaletteCollection(FName PaletteCollectionName)
 {
 	if (UnloadedPaletteCollections.Contains(PaletteCollectionName))
 	{
-		FSoundscapePaletteCollection PaletteCollectionToLoad = UnloadedPaletteCollections.FindRef(PaletteCollectionName);
-
 		UE_LOG(LogSoundscapeSubsystem, Verbose, TEXT("Attempting to load Palette Collection %s"), *PaletteCollectionName.ToString());
 
-		for (FSoftObjectPath& ObjPath : PaletteCollectionToLoad.SoundscapePaletteCollection)
+		FSoundscapePaletteCollection PaletteCollectionToLoad = UnloadedPaletteCollections.FindRef(PaletteCollectionName);
+
+		// Early out as false if Collection Empty
+		if (PaletteCollectionToLoad.SoundscapePaletteCollection.Num() <= 0)
 		{
-			if (UObject* PalettePath = ObjPath.TryLoad())
-			{
-				USoundscapePalette* SoundscapePalette = Cast<USoundscapePalette>(PalettePath);
+			UE_LOG(LogSoundscapeSubsystem, Verbose, TEXT("Attempted to load Palette Collection %s but Collection Empty"), *PaletteCollectionName.ToString());
 
-				if (SoundscapePalette)
-				{
-					// If palette is valid, add it to the Subsystem Collection
-					LoadedPaletteCollectionSet.Add(SoundscapePalette);
+			return false;
+		}
 
-					UE_LOG(LogSoundscapeSubsystem, Verbose, TEXT("Valid Soundscape Palette %s in Palette Collection %s"), *SoundscapePalette->GetFullName(), *PaletteCollectionName.ToString());
-				}
-				else
-				{
-					UE_LOG(LogSoundscapeSubsystem, Verbose, TEXT("Invalid Soundscape Palette in Palette Collection %s"), *PaletteCollectionName.ToString());
-				}
-			}
+
+		// Set up Asset Path Array
+		TArray<FSoftObjectPath> AssetsToLoad;
+		AssetsToLoad.Reserve(PaletteCollectionToLoad.SoundscapePaletteCollection.Num());
+
+		// Copy paths over to array
+		for (FSoftObjectPath AssetPath : PaletteCollectionToLoad.SoundscapePaletteCollection)
+		{
+			AssetsToLoad.Add(AssetPath);
 		}
 
+		// Create a callback delegate for completing the async load request
+		FStreamableDelegate AssetLoadingDelegate = FStreamableDelegate::CreateUObject(this, &USoundscapeSubsystem::CompleteAsyncLoading, PaletteCollectionName);
+
+		// Set up shared pointer for Asset Manager Async Load Request
+		TSharedPtr<FStreamableHandle> AssetHandle;
+		AssetHandle = UAssetManager::GetStreamableManager().RequestAsyncLoad(MoveTemp(AssetsToLoad), MoveTemp(AssetLoadingDelegate));
+
+		// Add Shared Pointer to map, allowing multiple load calls in a single frame, if necessary
+		ActiveAssetLoadHandles.Add(PaletteCollectionName, AssetHandle);
+
 		return true;
 	}
 
@@ -645,6 +681,48 @@ bool USoundscapeSubsystem::LoadPaletteCollection(FName PaletteCollectionName)
 	return false;
 }
 
+void USoundscapeSubsystem::CompleteAsyncLoading(FName PaletteCollectionName)
+{
+	FSoundscapePaletteCollection PaletteCollectionToLoad = UnloadedPaletteCollections.FindRef(PaletteCollectionName);
+
+	for (FSoftObjectPath& ObjPath : PaletteCollectionToLoad.SoundscapePaletteCollection)
+	{
+		if (UObject* Object = ObjPath.ResolveObject())
+		{
+			USoundscapePalette* SoundscapePalette = Cast<USoundscapePalette>(Object);
+
+			if (SoundscapePalette)
+			{
+				// If palette is valid, add it to the Subsystem Collection
+				LoadedPaletteCollectionSet.Add(SoundscapePalette);
+
+				UE_LOG(LogSoundscapeSubsystem, Verbose, TEXT("Valid Soundscape Palette %s in Palette Collection %s"), *SoundscapePalette->GetFullName(), *PaletteCollectionName.ToString());
+			}
+			else
+			{
+				UE_LOG(LogSoundscapeSubsystem, Verbose, TEXT("Invalid Soundscape Palette in Palette Collection %s"), *PaletteCollectionName.ToString());
+			}
+		}
+		else
+		{
+			UE_LOG(LogSoundscapeSubsystem, Verbose, TEXT("Failed to load %s"), *ObjPath.ToString());
+		}
+	}
+
+	// Release the handle pointer
+	if (TSharedPtr<FStreamableHandle>* AssetHandlePtr = ActiveAssetLoadHandles.Find(PaletteCollectionName))
+	{
+		TSharedPtr<FStreamableHandle> AssetHandle = *AssetHandlePtr;
+
+		if (AssetHandle.IsValid())
+		{
+			AssetHandle.Reset();
+		}
+
+		ActiveAssetLoadHandles.Remove(PaletteCollectionName);
+	}
+}
+
 bool USoundscapeSubsystem::UnloadPaletteCollection(FName PaletteCollectionName)
 {
 	if (UnloadedPaletteCollections.Find(PaletteCollectionName))
diff --git a/Engine/Plugins/Runtime/Soundscape/Source/SoundScape/Public/SoundscapeSubsystem.h b/Engine/Plugins/Runtime/Soundscape/Source/SoundScape/Public/SoundscapeSubsystem.h
--- a/Engine/Plugins/Runtime/Soundscape/Source/SoundScape/Public/SoundscapeSubsystem.h
+++ b/Engine/Plugins/Runtime/Soundscape/Source/SoundScape/Public/SoundscapeSubsystem.h
@@ -11,6 +11,7 @@
 
 class USoundscapePalette;
 class UActiveSoundscapePalette;
+struct FStreamableHandle;
 
 DECLARE_LOG_CATEGORY_EXTERN(LogSoundscapeSubsystem, Log, All);
 
@@ -245,7 +246,11 @@ private:
 	UPROPERTY()
 	TMap<TObjectPtr<USoundscapePalette>, TObjectPtr<UActiveSoundscapePalette>> ActivePalettes;
 
+	// Start Asynchronous Loading 
 	bool LoadPaletteCollection(FName PaletteCollectionName);
+	void CompleteAsyncLoading(FName PaletteCollectionName);
+
+	TMap<FName, TSharedPtr<FStreamableHandle>> ActiveAssetLoadHandles;
 
 	bool UnloadPaletteCollection(FName PaletteCollectionName);
 
