[Soundscape] Fixes
- Fixed issue where updates to Soundscape Palettes failed on BeginPlay
- Fixed issue on evaluating Color Point density when also using Move By Trace
- Fixed issue with draw debug cell call where box extents were twice as big
- Fixed multiple tooltip typos and copy paste errors

[RN] Fixed issue where Soundscape Palettes failed to activate properly on the same frame their state is set affecting Begin Play activation. Fixed an issue where Color Point Density was evaluated in the wrong location when combined with Move By Trace. Fixed issue where Draw Debug box extents were incorrectly sized. Fixed multiple tooltip typos and copy paste errors.

diff --git a/Engine/Plugins/Runtime/Soundscape/Source/SoundScape/Private/SoundscapeColor.cpp b/Engine/Plugins/Runtime/Soundscape/Source/SoundScape/Private/SoundscapeColor.cpp
--- a/Engine/Plugins/Runtime/Soundscape/Source/SoundScape/Private/SoundscapeColor.cpp
+++ b/Engine/Plugins/Runtime/Soundscape/Source/SoundScape/Private/SoundscapeColor.cpp
@@ -211,8 +211,11 @@ void UActiveSoundscapeColor::OnTraceCompleted(const FTraceHandle& Handle, FTrace
 			// Update sound location
 			NewSoundParams.Location = Data.OutHits[0].ImpactPoint;
 
-			// Spawn new sound at updated location
-			SpawnSoundscapeColorVoice();
+			if (EvaluateColorPointDensity())
+			{
+				// Spawn new sound at updated location
+				SpawnSoundscapeColorVoice();
+			}
 		}
 	}
 }
@@ -351,60 +354,25 @@ void UActiveSoundscapeColor::Update()
 		}
 	}
 
-
-
 	// Play Sound if needed and valid sound available
 	if (Sound && bNeedToSpawnSound && World)
 	{
 		// Calculate Spawn Parameters
 		CalculateSpawnParams();
 
-		USoundscapeSubsystem* SoundscapeSubsystem = nullptr;
-
-		// Set flag for can spawn
-		bool bCanSpawnSound = true;
-
-		// Get Soundscape Subsystem
-		if (UGameInstance* GameInstance = World->GetGameInstance())
-		{
-			SoundscapeSubsystem = GameInstance->GetSubsystem<USoundscapeSubsystem>();
-		}
-
-		// Check to see if Filter By Color Point Density is TRUE, determine if Soundscape Subsystem is still valid
-		if (SpawnBehavior.bFilterByColorPointDensity && SoundscapeSubsystem)
+		// Check to see if Position By Trace
+		if (SpawnBehavior.bPositionByTrace == false && EvaluateColorPointDensity())
 		{
-			// Cache Color Point Density of specified type
-			int32 ColorPointDensity = SoundscapeSubsystem->CheckColorPointDensity(NewSoundParams.Location, SpawnBehavior.ColorPoint);
-
-			// If Color Point Density is under threshold, set spawn flag to false
-			if (SpawnBehavior.MinColorPointNumber > ColorPointDensity)
-			{
-				bCanSpawnSound = false;
-			}
-
-			// Attempt Draw Debug Cell
-			SoundscapeSubsystem->DrawDebugCell(NewSoundParams.Location, bCanSpawnSound);
+			//Spawn Color Voice
+			SpawnSoundscapeColorVoice();
 		}
-
-		// If Sound can Spawn, Spawn
-		if(bCanSpawnSound)
+		else if (bTracing == false)
 		{
-			// If not filtering by Color Point Density, then check to see if Position By Trace
-			if (SpawnBehavior.bPositionByTrace == false)
-			{
-				//Spawn Color Voice
-				SpawnSoundscapeColorVoice();
-			}
-			else if (bTracing == false)
-			{
-				//If we're not currently mid-trace, spawn sound by trace
-				TraceHandle = SpawnSoundByTrace();
-			}
+			//If we're not currently mid-trace, spawn sound by trace
+			TraceHandle = SpawnSoundByTrace();
 		}
-
 	}
 
-
 	// No longer the first time
 	bFirstSpawn = false;
 }
@@ -493,6 +461,58 @@ void UActiveSoundscapeColor::CalculateSpawnParams()
 	}
 }
 
+bool UActiveSoundscapeColor::EvaluateColorPointDensity()
+{
+	// Check if we actually want to filter by color point density, otherwise return true
+	if (SpawnBehavior.bFilterByColorPointDensity)
+	{
+		// Get World
+		UWorld* World = GetWorld();
+
+		// Return false if failed to get world
+		if (!World)
+		{
+			return false;
+		}
+
+		// Get Soundscape Subsystem
+		UGameInstance* GameInstance = World->GetGameInstance();
+
+		if (!GameInstance)
+		{
+			return false;
+		}
+
+		// Check to see if Filter By Color Point Density is TRUE, determine if Soundscape Subsystem is still valid
+		if (USoundscapeSubsystem* SoundscapeSubsystem = GameInstance->GetSubsystem<USoundscapeSubsystem>())
+		{
+			// Cache Color Point Density of specified type
+			int32 ColorPointDensity = SoundscapeSubsystem->CheckColorPointDensity(NewSoundParams.Location, SpawnBehavior.ColorPoint);
+
+			// If Color Point Density is under threshold, set spawn flag to false
+			if (SpawnBehavior.MinColorPointNumber > ColorPointDensity)
+			{
+				// Attempt Draw Debug Cell
+				SoundscapeSubsystem->DrawDebugCell(NewSoundParams.Location, false);
+
+				// Failed to meet Color Point Density requirements, early out and do not spawn sound
+				return false;
+			}
+			
+			// Attempt Draw Debug Cell
+			SoundscapeSubsystem->DrawDebugCell(NewSoundParams.Location, true);
+
+			// We have met the threshold, return success
+			return true;
+		}
+
+		// If we failed getting the Game Instance, we do not want to return success
+		return false;
+	}
+
+	return true;
+}
+
 // Spawn a Color Voice
 void UActiveSoundscapeColor::SpawnSoundscapeColorVoice()
 {
diff --git a/Engine/Plugins/Runtime/Soundscape/Source/SoundScape/Private/SoundscapeSubsystem.cpp b/Engine/Plugins/Runtime/Soundscape/Source/SoundScape/Private/SoundscapeSubsystem.cpp
--- a/Engine/Plugins/Runtime/Soundscape/Source/SoundScape/Private/SoundscapeSubsystem.cpp
+++ b/Engine/Plugins/Runtime/Soundscape/Source/SoundScape/Private/SoundscapeSubsystem.cpp
@@ -588,11 +588,7 @@ bool USoundscapeSubsystem::AddPaletteCollection(FName PaletteCollectionName, FSo
 
 		UE_LOG(LogSoundscapeSubsystem, Verbose, TEXT("Palette Collection %s added"), *PaletteCollectionName.ToString());
 
-		bool bLoaded = LoadPaletteCollection(PaletteCollectionName);
-
-		UpdateState();
-
-		return bLoaded;
+		return LoadPaletteCollection(PaletteCollectionName);
 	}
 
 	UE_LOG(LogSoundscapeSubsystem, Verbose, TEXT("Add Palette Collection called with empty Palette Collection"));
@@ -721,6 +717,9 @@ void USoundscapeSubsystem::CompleteAsyncLoading(FName PaletteCollectionName)
 
 		ActiveAssetLoadHandles.Remove(PaletteCollectionName);
 	}
+
+	// Now that we have completed loading Palettes, we will call update to check whether new palettes need to be activated.
+	UpdateState();
 }
 
 bool USoundscapeSubsystem::UnloadPaletteCollection(FName PaletteCollectionName)
@@ -939,7 +938,7 @@ void USoundscapeSubsystem::DrawDebugCell(FVector Location, bool bSuccess)
 		}
 
 		FColor BoxColor = bSuccess ? FColor::Green : FColor::Red;
-		DrawDebugBox(GetWorld(), BoxDimensions.Key, BoxDimensions.Value, BoxColor, false, 5.0f, '\000', 5.0f);
+		DrawDebugBox(GetWorld(), BoxDimensions.Key, BoxDimensions.Value * 0.5f, BoxColor, false, 5.0f, '\000', 5.0f);
 	}
 }
 
diff --git a/Engine/Plugins/Runtime/Soundscape/Source/SoundScape/Public/SoundscapeColor.h b/Engine/Plugins/Runtime/Soundscape/Source/SoundScape/Public/SoundscapeColor.h
--- a/Engine/Plugins/Runtime/Soundscape/Source/SoundScape/Public/SoundscapeColor.h
+++ b/Engine/Plugins/Runtime/Soundscape/Source/SoundScape/Public/SoundscapeColor.h
@@ -21,8 +21,11 @@ DECLARE_DYNAMIC_MULTICAST_DELEGATE_OneParam(FOnSoundscapeColorParameterChanges,
 UENUM(BlueprintType)
 enum class ESoundscapeColorAltitudeClampMode : uint8
 {
-	Relative = 0 UMETA(DisplayName = "Relative Altitude"), // Altitude relative to the listener location
-	Absolute = 1 UMETA(DisplayName = "Absolute Altitude") // Altitude relative to the world coordinates
+	// Altitude relative to the listener location
+	Relative = 0 UMETA(DisplayName = "Relative Altitude"), 
+
+	// Altitude based on world coordinates
+	Absolute = 1 UMETA(DisplayName = "Absolute Altitude") 
 };
 
 // Struct 
@@ -31,27 +34,27 @@ struct SOUNDSCAPE_API FSoundscapeColorSpawnSettings
 {
 	GENERATED_BODY()
 
-	// Automatically Element Spawn when Element Returns Finished
+	// Include an additional delay for the first spawn attempt between MinFirstSpawnDelay and MaxFirstSpawnDelay times.
 	UPROPERTY(EditAnywhere, BlueprintReadWrite, Category = "SpawnBehavior")
 	bool bDelayFirstSpawn = false;
 
-	// Sets Minimum Time between Spawning Elements, checking if Elements are Finished, etc.
+	// Sets the minimum time delay before the first spawn
 	UPROPERTY(EditAnywhere, BlueprintReadWrite, Category = "SpawnBehavior", meta = (ClampMin = "0.0", UIMin = "0.0", EditCondition = "bDelayFirstSpawn"))
 	float MinFirstSpawnDelay = 1.0f;
 
-	// Sets the Maximum Time between Spawning Elements, checking if Elements are Finished, etc.
+	// Sets the maximum time delay before the first spawn
 	UPROPERTY(EditAnywhere, BlueprintReadWrite, Category = "SpawnBehavior", meta = (ClampMin = "0.0", UIMin = "0.0", EditCondition = "bDelayFirstSpawn"))
 	float MaxFirstSpawnDelay = 3.0f;
 
-	// When false, there will only be one Element spawned
+	// When true, color will continue to attempt to spawn sounds at an update frequency between the MinSpawnDelay and MaxSpawnDelay times. When false, it will only spawn once.
 	UPROPERTY(EditAnywhere, BlueprintReadWrite, Category = "SpawnBehavior")
 	bool bContinuouslyRespawn = false;
 
-	// Sets Minimum Time between Spawning Elements, checking if Elements are Finished, etc.
+	// Sets the minimum time between respawns
 	UPROPERTY(EditAnywhere, BlueprintReadWrite, Category = "SpawnBehavior", meta = (ClampMin = "0.0", UIMin = "0.0", EditCondition = "bContinuouslyRespawn"))
 	float MinSpawnDelay = 1.0f;
 
-	// Sets the Maximum Time between Spawning Elements, checking if Elements are Finished, etc.
+	// Sets the maximum time between respawns
 	UPROPERTY(EditAnywhere, BlueprintReadWrite, Category = "SpawnBehavior", meta = (ClampMin = "0.0", UIMin = "0.0", EditCondition = "bContinuouslyRespawn"))
 	float MaxSpawnDelay = 3.0f;
 
@@ -79,7 +82,7 @@ struct SOUNDSCAPE_API FSoundscapeColorSpawnSettings
 	UPROPERTY(EditAnywhere, BlueprintReadWrite, Category = "SpawnBehavior", meta = (ClampMin = "0.0", ClampMax = "180.0", UIMin = "0.0", UIMax = "180.0"))
 	float MaxSpawnAngle = 180.0f;
 
-	// When false, there will only be one Element spawned
+	// When true, spawn locations will clamp the z axis according to the ClampMode rules between the MinSpawnHeightClamp and the MaxSpawnHeightClamp
 	UPROPERTY(EditAnywhere, BlueprintReadWrite, Category = "SpawnBehavior")
 	bool bClampHeight = false;
 
@@ -120,7 +123,7 @@ struct SOUNDSCAPE_API FSoundscapeColorSpawnSettings
 	UPROPERTY(EditAnywhere, BlueprintReadWrite, AdvancedDisplay, Category = "SpawnBehavior")
 	bool bPositionByTrace = false;
 
-	// The Minimum nomralized directional vector in world orientation for the Trace to fire (default to down, z = -1.0f)
+	// The Minimum normalized directional vector in world orientation for the Trace to fire (default to down, z = -1.0f)
 	UPROPERTY(EditAnywhere, BlueprintReadWrite, AdvancedDisplay, Category = "SpawnBehavior", meta = (EditCondition = "bPositionByTrace"))
 	FVector TraceVector = { 0.0f, 0.0f, -1.0f };
 
@@ -156,7 +159,7 @@ struct SOUNDSCAPE_API FSoundscapeColorModulationSettings
 {
 	GENERATED_BODY()
 
-	// Add Volume Fades to the Re-trigger Sound playback
+	// An additional randomization of Volume between VolumeMin and VolumeMax on top of previous volume adjustments
 	UPROPERTY(EditAnywhere, BlueprintReadWrite, Category = "ModulationBehavior")
 	bool bRandomizeVolume = false;
 
@@ -168,7 +171,7 @@ struct SOUNDSCAPE_API FSoundscapeColorModulationSettings
 	UPROPERTY(EditAnywhere, BlueprintReadWrite, Category = "ModulationBehavior", meta = (ClampMin = "0.0", ClampMax = "4.0", UIMin = "0.0", UIMax = "4.0", SliderExponent = "6.0", EditCondition = "bRandomizeVolume"))
 	float VolumeMax = 1.05f;
 
-	// Add Volume Fades to the Re-trigger Sound playback
+	//  An additional randomization of Pitch between PitchMin and PitchMax on top of previous pitch adjustments
 	UPROPERTY(EditAnywhere, BlueprintReadWrite, Category = "ModulationBehavior")
 	bool bRandomizePitch = false;
 
@@ -432,6 +435,9 @@ private:
 	// Calculate spawning parameters
 	void CalculateSpawnParams();
 
+	// Evaluate color point density meets threshold, returns true if evaluation passes or if we are not filtering by color point density
+	bool EvaluateColorPointDensity();
+
 	// Spawn a Color Voice
 	void SpawnSoundscapeColorVoice();
 
